HTTP based chunk installer/stage downloader. Currently Disabled.

diff --git a/Engine/Source/Runtime/AssetRegistry/Private/AssetRegistry.cpp b/Engine/Source/Runtime/AssetRegistry/Private/AssetRegistry.cpp
--- a/Engine/Source/Runtime/AssetRegistry/Private/AssetRegistry.cpp
+++ b/Engine/Source/Runtime/AssetRegistry/Private/AssetRegistry.cpp
@@ -1048,7 +1048,7 @@ void FAssetRegistry::PrioritizeAssetInstall(const FAssetData& AssetData) const
 		return;
 	}
 
-	ChunkInstall->PrioritizeChunk(AssetData.ChunkIDs[0]);
+	ChunkInstall->PrioritizeChunk(AssetData.ChunkIDs[0], EChunkPriority::Immediate);
 }
 
 bool FAssetRegistry::AddPath(const FString& PathToAdd)
diff --git a/Engine/Source/Runtime/Core/Public/GenericPlatform/GenericPlatformChunkInstall.h b/Engine/Source/Runtime/Core/Public/GenericPlatform/GenericPlatformChunkInstall.h
--- a/Engine/Source/Runtime/Core/Public/GenericPlatform/GenericPlatformChunkInstall.h
+++ b/Engine/Source/Runtime/Core/Public/GenericPlatform/GenericPlatformChunkInstall.h
@@ -7,6 +7,8 @@
 
 #pragma once
 
+#include "ModuleInterface.h"
+
 DECLARE_LOG_CATEGORY_EXTERN(LogChunkInstaller, Log, All);
 
 namespace EChunkLocation
@@ -34,6 +36,15 @@ namespace EChunkInstallSpeed
 	};
 }
 
+namespace EChunkPriority
+{
+	enum Type
+	{
+		Immediate,	// Chunk install is of highest priority, this can cancel lower priority installs.
+		High	 ,	// Chunk is probably required soon so grab is as soon as possible.
+		Low		 ,	// Install this chunk only when other chunks are not needed.
+	};
+}
 
 namespace EChunkProgressReportingType
 {
@@ -44,6 +55,18 @@ namespace EChunkProgressReportingType
 	};
 }
 
+/**
+ * Platform Chunk Install Module Interface
+ */
+class IPlatformChunkInstallModule : public IModuleInterface
+{
+public:
+
+	virtual IPlatformChunkInstall* GetPlatformChunkInstall() = 0;
+};
+
+
+DECLARE_DELEGATE_OneParam(FPlatformChunkInstallCompleteDelegate, uint32);
 
 /**
 * Interface for platform specific chunk based install
@@ -90,11 +113,12 @@ public:
 	virtual bool SetInstallSpeed( EChunkInstallSpeed::Type InstallSpeed ) = 0;
 	
 	/**
-	 * Hint to the installer that we would like to prioritize a specific chunk (moves it to the head of the list.)
+	 * Hint to the installer that we would like to prioritize a specific chunk
 	 * @param ChunkID		The id of the chunk to prioritize.
+	 * @param Priority		The priority for the chunk.
 	 * @return				false if the operation is not allowed or the chunk doesn't exist, otherwise true.
 	 **/
-	virtual bool PrioritizeChunk( uint32 ChunkID ) = 0;
+	virtual bool PrioritizeChunk( uint32 ChunkID, EChunkPriority::Type Priority ) = 0;
 
 	/**
 	 * For platforms that support emulation of the Chunk install.  Starts transfer of the next chunk.
@@ -102,6 +126,21 @@ public:
 	 * @return				true if the opreation succeeds.
 	 **/
 	virtual bool DebugStartNextChunk() = 0;
+
+	/** 
+	 * Request a delegate callback on chunk install completion. Request may not be respected.
+	 * @param ChunkID		The id of the chunk of interest.
+	 * @param Delegate		The delegate to call when the chunk is installed.
+	 * @return				False if the delegate was not registered. True on success.
+	 */
+	virtual bool SetChunkInstallDelgate( uint32 ChunkID, FPlatformChunkInstallCompleteDelegate Delegate ) = 0;
+
+	/**
+	* Remove a delegate callback on chunk install completion.
+	* @param ChunkID		The id of the chunk of interest.
+	* @param Delegate		The delegate to remove.
+	*/
+	virtual void RemoveChunkInstallDelgate(uint32 ChunkID, FPlatformChunkInstallCompleteDelegate Delegate) = 0;
 };
 
 
@@ -175,7 +214,7 @@ public:
 	 * @param ChunkID		The id of the chunk to prioritize.
 	 * @return				false if the operation is not allowed or the chunk doesn't exist, otherwise true.
 	 **/
-	virtual bool PrioritizeChunk( uint32 ChunkID ) override
+	virtual bool PrioritizeChunk( uint32 ChunkID, EChunkPriority::Type Priority ) override
 	{
 		return false;
 	}
@@ -189,4 +228,26 @@ public:
 	{
 		return true;
 	}
+
+	/**
+	* Request a delegate callback on chunk install completion. Request may not be respected.
+	* @param ChunkID		The id of the chunk of interest.
+	* @param Delegate		The delegate when the chunk is installed.
+	* @return				False if the delegate was not registered or the chunk is already installed. True on success.
+	*/
+	virtual bool SetChunkInstallDelgate(uint32 ChunkID, FPlatformChunkInstallCompleteDelegate Delegate)
+	{
+		return false;
+	}
+
+	/**
+	* Remove a delegate callback on chunk install completion.
+	* @param ChunkID		The id of the chunk of interest.
+	* @param Delegate		The delegate to remove.
+	* @return				False if the delegate was not registered with a call to SetChunkInstallDelgate. True on success.
+	*/
+	virtual void RemoveChunkInstallDelgate(uint32 ChunkID, FPlatformChunkInstallCompleteDelegate Delegate)
+	{
+		return;
+	}
 };
diff --git a/Engine/Source/Runtime/Core/Public/Misc/CallbackDevice.h b/Engine/Source/Runtime/Core/Public/Misc/CallbackDevice.h
--- a/Engine/Source/Runtime/Core/Public/Misc/CallbackDevice.h
+++ b/Engine/Source/Runtime/Core/Public/Misc/CallbackDevice.h
@@ -62,7 +62,7 @@ public:
 	#endif	//WITH_EDITOR
 
 	// delegate type for prompting the pak system to mount a new pak
-	DECLARE_DELEGATE_RetVal_OneParam(bool, FOnMountPak, const FString& );
+	DECLARE_DELEGATE_RetVal_TwoParams(bool, FOnMountPak, const FString&, uint32);
 
 	// Callback for PER_MODULE_BOILERPLATE macro's GSerialNumberBlocksForDebugVisualizers
 	DECLARE_DELEGATE_RetVal(int32***, FGetSerialNumberBlocksForDebugVisualizersDelegate);
diff --git a/Engine/Source/Runtime/HTTPChunkInstaller/HTTPChunkInstaller.Build.cs b/Engine/Source/Runtime/HTTPChunkInstaller/HTTPChunkInstaller.Build.cs
new file mode 100644
--- /dev/null
+++ b/Engine/Source/Runtime/HTTPChunkInstaller/HTTPChunkInstaller.Build.cs
@@ -0,0 +1,17 @@
+// Copyright 1998-2014 Epic Games, Inc. All Rights Reserved.
+
+using UnrealBuildTool;
+
+public class HTTPChunkInstaller : ModuleRules
+{
+	public HTTPChunkInstaller(TargetInfo Target)
+	{
+		PrivateDependencyModuleNames.Add("Core");
+        PrivateDependencyModuleNames.Add("Engine");
+        PrivateDependencyModuleNames.Add("PakFile");
+		PrivateDependencyModuleNames.Add("HTTP");
+        PrivateDependencyModuleNames.Add("BuildPatchServices");
+        PrivateDependencyModuleNames.Add("OnlineSubsystem");
+        PrivateDependencyModuleNames.Add("OnlineSubsystemUtils");
+	}
+}
\ No newline at end of file
diff --git a/Engine/Source/Runtime/HTTPChunkInstaller/Private/ChunkInstall.h b/Engine/Source/Runtime/HTTPChunkInstaller/Private/ChunkInstall.h
new file mode 100644
--- /dev/null
+++ b/Engine/Source/Runtime/HTTPChunkInstaller/Private/ChunkInstall.h
@@ -0,0 +1,101 @@
+// Copyright 1998-2014 Epic Games, Inc. All Rights Reserved.
+
+#pragma once
+
+
+class FChunkInstallTask : public FNonAbandonableTask
+{
+public:
+	/** Input parameters */
+	FString							ManifestPath;
+	FString							SrcDir;
+	FString							DestDir;
+	IBuildPatchServicesModule*		BPSModule;
+	IBuildManifestPtr				BuildManifest;
+	bool							bCopy;
+	const TArray<FString>*			CurrentMountPaks;
+	/** Output */
+	TArray<FString>					MountedPaks;
+
+	FChunkInstallTask()
+	{}
+
+	void Init(FString InManifestPath, FString InSrcDir, FString InDestDir, IBuildPatchServicesModule* InBPSModule, IBuildManifestRef InBuildManifest, const TArray<FString>& InCurrentMountedPaks, bool bInCopy)
+	{
+		ManifestPath = InManifestPath;
+		SrcDir = InSrcDir;
+		DestDir = InDestDir;
+		BPSModule = InBPSModule;
+		BuildManifest = InBuildManifest;
+		CurrentMountPaks = &InCurrentMountedPaks;
+		bCopy = bInCopy;
+
+		MountedPaks.Reset();
+	}
+
+	void DoWork()
+	{
+		// Helper class to find all pak files.
+		class FPakSearchVisitor : public IPlatformFile::FDirectoryVisitor
+		{
+			TArray<FString>& FoundPakFiles;
+		public:
+			FPakSearchVisitor(TArray<FString>& InFoundPakFiles)
+				: FoundPakFiles(InFoundPakFiles)
+			{}
+			virtual bool Visit(const TCHAR* FilenameOrDirectory, bool bIsDirectory)
+			{
+				if (bIsDirectory == false)
+				{
+					FString Filename(FilenameOrDirectory);
+					if (FPaths::GetExtension(Filename) == TEXT("pak"))
+					{
+						FoundPakFiles.Add(Filename);
+					}
+				}
+				return true;
+			}
+		};
+
+		check(CurrentMountPaks);
+
+		BPSModule->SaveManifestToFile(ManifestPath, BuildManifest.ToSharedRef());
+		if (bCopy)
+		{
+			TArray<FString> PakFiles;
+			IPlatformFile& PlatformFile = FPlatformFileManager::Get().GetPlatformFile();
+			if (PlatformFile.DirectoryExists(*DestDir))
+			{
+				PlatformFile.DeleteDirectoryRecursively(*DestDir);
+			}
+			PlatformFile.CreateDirectoryTree(*DestDir);
+			if (PlatformFile.CopyDirectoryTree(*DestDir, *SrcDir, true))
+			{
+				PlatformFile.DeleteDirectoryRecursively(*SrcDir);
+			}
+			// Find all pak files.
+			FPakSearchVisitor Visitor(PakFiles);
+			PlatformFile.IterateDirectoryRecursively(*DestDir, Visitor);
+			auto PakReadOrderField = BuildManifest->GetCustomField("PakReadOrdering");
+			uint32 PakReadOrder = PakReadOrderField.IsValid() ? (uint32)PakReadOrderField->AsInteger() : 0;
+			for (uint32 PakIndex = 0, PakCount = PakFiles.Num(); PakIndex < PakCount; ++PakIndex)
+			{
+				if (!CurrentMountPaks->Contains(PakFiles[PakIndex]) && !MountedPaks.Contains(PakFiles[PakIndex]))
+				{
+					if (FCoreDelegates::OnMountPak.IsBound())
+					{
+						FCoreDelegates::OnMountPak.Execute(PakFiles[PakIndex], PakReadOrder);
+						MountedPaks.Add(PakFiles[PakIndex]);
+					}
+				}
+			}
+			//Register the install
+			BPSModule->RegisterAppInstallation(BuildManifest.ToSharedRef(), DestDir);
+		}
+	}
+
+	static const TCHAR *Name()
+	{
+		return TEXT("FChunkDescovery");
+	}
+};
diff --git a/Engine/Source/Runtime/HTTPChunkInstaller/Private/ChunkSetup.h b/Engine/Source/Runtime/HTTPChunkInstaller/Private/ChunkSetup.h
new file mode 100644
--- /dev/null
+++ b/Engine/Source/Runtime/HTTPChunkInstaller/Private/ChunkSetup.h
@@ -0,0 +1,253 @@
+// Copyright 1998-2014 Epic Games, Inc. All Rights Reserved.
+
+#pragma once
+
+// Helper class to find all pak files.
+class FFileSearchVisitor : public IPlatformFile::FDirectoryVisitor
+{
+	FString			 FileWildcard;
+	TArray<FString>& FoundFiles;
+public:
+	FFileSearchVisitor(FString InFileWildcard, TArray<FString>& InFoundFiles)
+		: FileWildcard(InFileWildcard)
+		, FoundFiles(InFoundFiles)
+	{}
+	virtual bool Visit(const TCHAR* FilenameOrDirectory, bool bIsDirectory)
+	{
+		if (bIsDirectory == false)
+		{
+			FString Filename(FilenameOrDirectory);
+			if (Filename.MatchesWildcard(FileWildcard))
+			{
+				FoundFiles.Add(Filename);
+			}
+		}
+		return true;
+	}
+};
+
+class IBuildPatchServicesModule;
+
+class FChunkSetupTask : public FNonAbandonableTask, public IPlatformFile::FDirectoryVisitor
+{
+public:
+	/** Input parameters */
+	IBuildPatchServicesModule*	BPSModule;
+	FString						InstallDir; // Intermediate directory where installed chunks may be waiting
+	FString						ContentDir; // Directory where installed chunks need to live to be mounted
+	const TArray<FString>*		CurrentMountPaks; // Array of already mounted Paks
+	/** Output */
+	TArray<FString>							MountedPaks;
+	TMultiMap<uint32, IBuildManifestPtr>	InstalledChunks;
+	/** Working */
+	TArray<FString>		FoundPaks;
+	TArray<FString>		FoundManifests;
+	FFileSearchVisitor	PakVisitor;
+	FFileSearchVisitor	ManifestVisitor;
+	TArray<FString>		DirectoriesToRemove;
+	uint32				Pass;
+
+	FChunkSetupTask()
+		: BPSModule(nullptr)
+		, CurrentMountPaks(nullptr)
+		, PakVisitor(TEXT("*.pak"), FoundPaks)
+		, ManifestVisitor(TEXT("*.manifest"), FoundManifests)
+	{}
+
+	void Init(IBuildPatchServicesModule* InBPSModule, FString InInstallDir, FString InContentDir, const TArray<FString>& InCurrentMountedPaks)
+	{
+		BPSModule = InBPSModule;
+		InstallDir = InInstallDir;
+		ContentDir = InContentDir;
+		CurrentMountPaks = &InCurrentMountedPaks;
+
+		Pass = 0;
+		InstalledChunks.Reset();
+		MountedPaks.Reset();
+		FoundManifests.Reset();
+		FoundPaks.Reset();
+		DirectoriesToRemove.Reset();
+	}
+
+	void DoWork()
+	{
+		IPlatformFile& PlatformFile = FPlatformFileManager::Get().GetPlatformFile();
+		PlatformFile.IterateDirectory(*InstallDir, *this);
+		for (const auto& ToRemove : DirectoriesToRemove)
+		{
+			PlatformFile.DeleteDirectoryRecursively(*ToRemove);
+		}
+		++Pass;
+		PlatformFile.IterateDirectory(*ContentDir, *this);
+	}
+
+	bool Visit(const TCHAR* FilenameOrDirectory, bool bIsDirectory) override
+	{
+		if (!bIsDirectory)
+		{
+			return true;
+		}
+		IPlatformFile& PlatformFile = FPlatformFileManager::Get().GetPlatformFile();
+
+		FoundManifests.Reset();
+		PlatformFile.IterateDirectory(FilenameOrDirectory, ManifestVisitor);
+		if (FoundManifests.Num() == 0)
+		{
+			return true;
+		}
+		// Only allow one manifest per folder, any more suggests corruption so mark the folder for delete
+		if (FoundManifests.Num() > 1)
+		{
+			DirectoriesToRemove.Add(FilenameOrDirectory);
+			return true;
+		}
+		// Load the manifest, so that can be classed as installed
+		auto Manifest = BPSModule->LoadManifestFromFile(FoundManifests[0]);
+		if (!Manifest.IsValid())
+		{
+			//Something is wrong, suggests corruption so mark the folder for delete
+			DirectoriesToRemove.Add(FilenameOrDirectory);
+			return true;
+		}
+		auto ChunkIDField = Manifest->GetCustomField("ChunkID");
+		if (!ChunkIDField.IsValid())
+		{
+			//Something is wrong, suggests corruption so mark the folder for delete
+			DirectoriesToRemove.Add(FilenameOrDirectory);
+			return true;
+		}
+		auto ChunkPatchField = Manifest->GetCustomField("bIsPatch");
+		bool bIsPatch = ChunkPatchField.IsValid() ? ChunkPatchField->AsString() == TEXT("true") : false;
+		uint32 ChunkID = (uint32)ChunkIDField->AsInteger();
+
+		if (Pass == 1)
+		{
+			InstalledChunks.AddUnique(ChunkID, Manifest);
+		}
+		else
+		{
+			FString ChunkFdrName = FString::Printf(TEXT("%s%d"), !bIsPatch ? TEXT("base") : TEXT("patch"), ChunkID);
+			FString DestDir = *FPaths::Combine(*ContentDir, *ChunkFdrName);
+			if (PlatformFile.DirectoryExists(*DestDir))
+			{
+				PlatformFile.DeleteDirectoryRecursively(*DestDir);
+			}
+			PlatformFile.CreateDirectoryTree(*DestDir);
+			if (PlatformFile.CopyDirectoryTree(*DestDir, FilenameOrDirectory, true))
+			{
+				DirectoriesToRemove.Add(FilenameOrDirectory);
+			}
+		}
+
+		return true;
+	}
+
+	static const TCHAR *Name()
+	{
+		return TEXT("FChunkSetup");
+	}
+};
+
+class FChunkMountTask : public FNonAbandonableTask, public IPlatformFile::FDirectoryVisitor
+{
+public:
+	/** Input parameters */
+	IBuildPatchServicesModule*	BPSModule;
+	FString						ContentDir; // Directory where installed chunks need to live to be mounted
+	const TArray<FString>*		CurrentMountPaks; // Array of already mounted Paks
+	/** Output */
+	TArray<FString>							MountedPaks;
+	/** Working */
+	TArray<FString>		FoundPaks;
+	TArray<FString>		FoundManifests;
+	FFileSearchVisitor	PakVisitor;
+	FFileSearchVisitor	ManifestVisitor;
+
+	FChunkMountTask()
+		: BPSModule(nullptr)
+		, CurrentMountPaks(nullptr)
+		, PakVisitor(TEXT("*.pak"), FoundPaks)
+		, ManifestVisitor(TEXT("*.manifest"), FoundManifests)
+	{}
+
+	void Init(IBuildPatchServicesModule* InBPSModule, FString InContentDir, const TArray<FString>& InCurrentMountedPaks)
+	{
+		BPSModule = InBPSModule;
+		ContentDir = InContentDir;
+		CurrentMountPaks = &InCurrentMountedPaks;
+
+		MountedPaks.Reset();
+		FoundManifests.Reset();
+		FoundPaks.Reset();
+	}
+
+	void DoWork()
+	{
+		IPlatformFile& PlatformFile = FPlatformFileManager::Get().GetPlatformFile();
+		PlatformFile.IterateDirectory(*ContentDir, *this);
+	}
+
+	bool Visit(const TCHAR* FilenameOrDirectory, bool bIsDirectory) override
+	{
+		if (!bIsDirectory)
+		{
+			return true;
+		}
+		IPlatformFile& PlatformFile = FPlatformFileManager::Get().GetPlatformFile();
+
+		FoundManifests.Reset();
+		PlatformFile.IterateDirectory(FilenameOrDirectory, ManifestVisitor);
+		if (FoundManifests.Num() == 0)
+		{
+			return true;
+		}
+		// Only allow one manifest per folder, any more suggests corruption so ignore
+		if (FoundManifests.Num() > 1)
+		{
+			return true;
+		}
+		// Load the manifest, so that can be classed as installed
+		auto Manifest = BPSModule->LoadManifestFromFile(FoundManifests[0]);
+		if (!Manifest.IsValid())
+		{
+			//Something is wrong, suggests corruption so ignore
+			return true;
+		}
+		auto ChunkIDField = Manifest->GetCustomField("ChunkID");
+		if (!ChunkIDField.IsValid())
+		{
+			//Something is wrong, suggests corruption so ignore
+			return true;
+		}
+
+		FoundPaks.Reset();
+		PlatformFile.IterateDirectoryRecursively(FilenameOrDirectory, PakVisitor);
+		if (FoundPaks.Num() == 0)
+		{
+			return true;
+		}
+		auto PakReadOrderField = Manifest->GetCustomField("PakReadOrdering");
+		uint32 PakReadOrder = PakReadOrderField.IsValid() ? (uint32)PakReadOrderField->AsInteger() : 0;
+		for (const auto& PakPath : FoundPaks)
+		{
+			//Note: A side effect here is that any unmounted paks get mounted. This is desirable as
+			// any previously installed chunks get mounted 
+			if (!CurrentMountPaks->Contains(PakPath) && !MountedPaks.Contains(PakPath))
+			{
+				if (FCoreDelegates::OnMountPak.IsBound())
+				{
+					FCoreDelegates::OnMountPak.Execute(PakPath, PakReadOrder);
+					MountedPaks.Add(PakPath);
+					//Register the install
+					BPSModule->RegisterAppInstallation(Manifest.ToSharedRef(), FilenameOrDirectory);
+				}
+			}
+		}
+		return true;
+	}
+
+	static const TCHAR *Name()
+	{
+		return TEXT("FChunkSetup");
+	}
+};
\ No newline at end of file
diff --git a/Engine/Source/Runtime/HTTPChunkInstaller/Private/HTTPChunkInstaller.cpp b/Engine/Source/Runtime/HTTPChunkInstaller/Private/HTTPChunkInstaller.cpp
new file mode 100644
--- /dev/null
+++ b/Engine/Source/Runtime/HTTPChunkInstaller/Private/HTTPChunkInstaller.cpp
@@ -0,0 +1,614 @@
+// Copyright 1998-2014 Epic Games, Inc. All Rights Reserved.
+
+#ifndef PLACEHOLDER_DLC_WORK
+#	define PLACEHOLDER_DLC_WORK 0
+#endif
+
+#include "HTTPChunkInstallerPrivatePCH.h"
+#include "HTTPChunkInstaller.h"
+#include "ChunkInstall.h"
+#include "UniquePtr.h"
+#include "LocalTitleFile.h"
+
+// helper to grab the installer service
+static IBuildPatchServicesModule* GetBuildPatchServices()
+{
+	static IBuildPatchServicesModule* BuildPatchServices = nullptr;
+	if (BuildPatchServices == nullptr)
+	{
+		BuildPatchServices = &FModuleManager::LoadModuleChecked<IBuildPatchServicesModule>(TEXT("BuildPatchServices"));
+	}
+	return BuildPatchServices;
+}
+
+// Helper class to find all pak file manifests.
+class FChunkSearchVisitor: public IPlatformFile::FDirectoryVisitor
+{
+public:
+	TArray<FString>				PakManifests;
+
+	FChunkSearchVisitor()
+	{}
+
+	virtual bool Visit(const TCHAR* FilenameOrDirectory,bool bIsDirectory)
+	{
+		if(bIsDirectory == false)
+		{
+			FString Filename(FilenameOrDirectory);
+			if(FPaths::GetBaseFilename(Filename).MatchesWildcard("*.manifest"))
+			{
+				PakManifests.AddUnique(Filename);
+			}
+		}
+		return true;
+	}
+};
+
+
+FHTTPChunkInstall::FHTTPChunkInstall()
+	: InstallingChunkID(-1)
+	, InstallerState(ChunkInstallState::Setup)
+	, InstallSpeed(EChunkInstallSpeed::Fast)
+	, bDebugNoInstalledRequired(false)
+{
+	BPSModule = GetBuildPatchServices();
+
+#if !UE_BUILD_SHIPPING
+	const TCHAR* CmdLine = FCommandLine::Get();
+	bool Result = FParse::Param(CmdLine,TEXT("Pak")) || FParse::Param(CmdLine,TEXT("Signedpak")) || FParse::Param(CmdLine,TEXT("Signed"));
+	if(!FPlatformProperties::RequiresCookedData() || !Result || FParse::Param(CmdLine,TEXT("NoPak")) || FParse::Param(CmdLine, TEXT("NoChunkInstall")))
+	{		
+		bDebugNoInstalledRequired = true;
+	}
+#endif
+
+	// Grab the title file interface
+	FString TitleFileSource;
+	bool bValidTitleFileSource = GConfig->GetString(TEXT("HTTPChunkInstall"), TEXT("TitleFileSource"), TitleFileSource, GEngineIni);
+	if (bValidTitleFileSource && TitleFileSource != TEXT("Local"))
+	{
+		OnlineTitleFile = Online::GetTitleFileInterface(*TitleFileSource);
+	}
+	else
+	{
+		OnlineTitleFile = MakeShareable(new FLocalTitleFile(FPaths::GameConfigDir()));
+	}
+	CloudDir	= FPaths::Combine(*FPaths::GameContentDir()	, TEXT("Cloud"));
+	StageDir	= FPaths::Combine(*FPaths::GameSavedDir()	, TEXT("Chunks"), TEXT("Staged"));
+	InstallDir	= FPaths::Combine(*FPaths::GameSavedDir()	, TEXT("Chunks"), TEXT("Installed"));
+	BackupDir	= FPaths::Combine(*FPaths::GameSavedDir()	, TEXT("Chunks"), TEXT("Backup"));
+	ContentDir  = FPaths::Combine(*FPaths::GameContentDir() , TEXT("Chunks"));
+
+	FString TmpString1;
+	FString TmpString2;
+	if (GConfig->GetString(TEXT("HTTPChunkInstall"), TEXT("CloudDirectory"), TmpString1, GEngineIni))
+	{
+		CloudDir = TmpString1;
+	}
+	if ((GConfig->GetString(TEXT("HTTPChunkInstall"), TEXT("CloudProtocol"), TmpString1, GEngineIni)) && (GConfig->GetString(TEXT("HTTPChunkInstall"), TEXT("CloudDomain"), TmpString2, GEngineIni)))
+	{
+		CloudDir = FString::Printf(TEXT("%s://%s"), *TmpString1, *TmpString2);
+	}
+	if (GConfig->GetString(TEXT("HTTPChunkInstall"), TEXT("StageDirectory"), TmpString1, GEngineIni))
+	{
+		StageDir = TmpString1;
+	}
+	if (GConfig->GetString(TEXT("HTTPChunkInstall"), TEXT("InstallDirectory"), TmpString1, GEngineIni))
+	{
+		InstallDir = TmpString1;
+	}
+	if (GConfig->GetString(TEXT("HTTPChunkInstall"), TEXT("BackupDirectory"), TmpString1, GEngineIni))
+	{
+		BackupDir = TmpString1;
+	}
+	if (GConfig->GetString(TEXT("HTTPChunkInstall"), TEXT("ContentDirectory"), TmpString1, GEngineIni))
+	{
+		ContentDir = TmpString1;
+	}
+
+	bFirstRun = true;
+}
+
+
+FHTTPChunkInstall::~FHTTPChunkInstall()
+{
+	if (InstallService.IsValid())
+	{
+		InstallService->CancelInstall();
+		InstallService.Reset();
+	}
+}
+
+bool FHTTPChunkInstall::Tick(float DeltaSeconds)
+{
+	switch (InstallerState)
+	{
+	case ChunkInstallState::Setup:
+		{
+			if (InstallSpeed != EChunkInstallSpeed::Paused)
+			{
+				ChunkSetupTask.GetTask().Init(BPSModule, InstallDir, ContentDir, MountedPaks);
+				ChunkSetupTask.StartBackgroundTask();
+				InstallerState = ChunkInstallState::SetupWait;
+			}
+		} break;
+	case ChunkInstallState::SetupWait:
+		{
+			if (ChunkSetupTask.IsDone())
+			{
+				for (auto It = ChunkSetupTask.GetTask().InstalledChunks.CreateConstIterator(); It; ++It)
+				{
+					InstalledManifests.Add(It.Key(), It.Value());
+				}
+				MountedPaks.Append(ChunkSetupTask.GetTask().MountedPaks);
+				InstallerState = ChunkInstallState::QueryRemoteManifests;
+			}
+		} break;
+	case ChunkInstallState::QueryRemoteManifests:
+		{			
+			//Now query the title file service for the chunk manifests. This should return the list of expected chunk manifests
+			if (!OnlineTitleFile.IsValid())
+			{
+				break;
+			}
+			OnlineTitleFile->AddOnEnumerateFilesCompleteDelegate(FOnEnumerateFilesCompleteDelegate::CreateRaw(this, &FHTTPChunkInstall::OSSEnumerateFilesComplete));
+			OnlineTitleFile->AddOnReadFileCompleteDelegate(FOnReadFileCompleteDelegate::CreateRaw(this, &FHTTPChunkInstall::OSSReadFileComplete));
+			OnlineTitleFile->ClearFiles();
+			InstallerState = ChunkInstallState::RequestingTitleFiles;
+			OnlineTitleFile->EnumerateFiles();
+		} break;
+	case ChunkInstallState::SearchTitleFiles:
+		{
+			FString CleanName;
+			TArray<FCloudFileHeader> FileList;
+			TitleFilesToRead.Reset();
+			RemoteManifests.Reset();
+			OnlineTitleFile->GetFileList(FileList);
+			for (int32 FileIndex = 0, FileCount = FileList.Num(); FileIndex < FileCount; ++FileIndex)
+			{
+				if (FileList[FileIndex].FileName.MatchesWildcard(TEXT("*.manifest")))
+				{
+					TitleFilesToRead.Add(FileList[FileIndex]);
+				}
+			}
+			InstallerState = ChunkInstallState::ReadTitleFiles;
+		} break;
+	case ChunkInstallState::ReadTitleFiles:
+		{
+			if (TitleFilesToRead.Num() > 0)
+			{
+				OnlineTitleFile->ReadFile(TitleFilesToRead[0].DLName);
+			}
+			else
+			{
+				InstallerState = ChunkInstallState::Idle;
+			}
+		} break;
+	case ChunkInstallState::ReadComplete:
+		{
+			if (OnlineTitleFile->GetFileContents(TitleFilesToRead[0].DLName, FileContentBuffer))
+			{
+				ParseTitleFileManifest();
+				// Even if the Parse failed remove the file from the list
+				TitleFilesToRead.RemoveAt(0);
+			}
+			if (TitleFilesToRead.Num() == 0)
+			{
+				if (bFirstRun)
+				{
+					ChunkMountTask.GetTask().Init(BPSModule, ContentDir, MountedPaks);
+					ChunkMountTask.StartBackgroundTask();
+				}
+				InstallerState = ChunkInstallState::PostSetup;
+			}
+			else
+			{
+				InstallerState = ChunkInstallState::ReadTitleFiles;
+			}
+		} break;
+	case ChunkInstallState::PostSetup:
+		{
+			if (bFirstRun)
+			{
+				if (ChunkMountTask.IsDone())
+				{
+					MountedPaks.Append(ChunkMountTask.GetTask().MountedPaks);
+					bFirstRun = false;
+				}
+			}
+			else
+			{
+				InstallerState = ChunkInstallState::Idle;
+			}
+		} break;
+	case ChunkInstallState::Idle: 
+		{
+			UpdatePendingInstallQueue();
+		} break;
+	case ChunkInstallState::CopyToContent:
+		{
+			if (!ChunkCopyInstall.IsDone() || !InstallService->IsComplete())
+			{
+				break;
+			}
+			check(InstallingChunkID != -1);
+			if (InstallService.IsValid())
+			{
+				InstallService.Reset();
+			}
+			check(RemoteManifests.Find(InstallingChunkID));
+			InstalledManifests.Add(InstallingChunkID, InstallingChunkManifest);
+			RemoteManifests.Remove(InstallingChunkID, InstallingChunkManifest);
+			MountedPaks.Append(ChunkCopyInstall.GetTask().MountedPaks);
+			if (!RemoteManifests.Contains(InstallingChunkID))
+			{
+				// No more manifests relating to the chunk ID are left to install.
+				// Inform any listeners that the install has been completed.
+				FPlatformChunkInstallCompleteMultiDelegate* FoundDelegate = DelegateMap.Find(InstallingChunkID);
+				if (FoundDelegate)
+				{
+					FoundDelegate->Broadcast(InstallingChunkID);
+				}
+			}
+			InstallingChunkID = -1;
+			InstallingChunkManifest.Reset();
+			InstallerState = ChunkInstallState::Idle;
+		} break;
+	case ChunkInstallState::Installing:
+	case ChunkInstallState::RequestingTitleFiles:
+	case ChunkInstallState::WaitingOnRead:
+	default:
+		break;
+	}
+
+	return true;
+}
+
+void FHTTPChunkInstall::UpdatePendingInstallQueue()
+{
+	if (InstallingChunkID != -1
+#if !UE_BUILD_SHIPPING
+		|| bDebugNoInstalledRequired
+#endif
+	)
+	{
+		return;
+	}
+
+	check(!InstallService.IsValid());
+	IPlatformFile& PlatformFile = FPlatformFileManager::Get().GetPlatformFile();
+	bool bPatch = false;
+	while (PriorityQueue.Num() > 0 && InstallerState != ChunkInstallState::Installing)
+	{
+		const FChunkPrio& NextChunk = PriorityQueue[0];
+		TArray<IBuildManifestPtr> FoundChunkManifests;
+		RemoteManifests.MultiFind(NextChunk.ChunkID, FoundChunkManifests);
+		if (FoundChunkManifests.Num() > 0)
+		{
+			auto ChunkManifest = FoundChunkManifests[0];
+			auto ChunkIDField = ChunkManifest->GetCustomField("ChunkID");
+			if (ChunkIDField.IsValid())
+			{
+				InstallingChunkID = NextChunk.ChunkID;
+				InstallingChunkManifest = ChunkManifest;
+				auto PatchField = ChunkManifest->GetCustomField("bIsPatch");
+				bool bIsPatch = PatchField.IsValid() ? PatchField->AsString() == TEXT("true") : false;
+				auto ChunkFolderName = FString::Printf(TEXT("%s%d"), !bIsPatch ? TEXT("base") : TEXT("patch"), ChunkIDField->AsInteger());
+				auto ChunkInstallDir = FPaths::Combine(*InstallDir, *ChunkFolderName);
+				auto ChunkStageDir = FPaths::Combine(*StageDir, *ChunkFolderName);
+				BPSModule->SetCloudDirectory(CloudDir);
+				BPSModule->SetStagingDirectory(ChunkStageDir);
+				InstallService = BPSModule->StartBuildInstall(nullptr, ChunkManifest, ChunkInstallDir, FBuildPatchBoolManifestDelegate::CreateRaw(this, &FHTTPChunkInstall::OSSInstallComplete));
+				if (InstallSpeed == EChunkInstallSpeed::Paused && !InstallService->IsPaused())
+				{
+					InstallService->TogglePauseInstall();
+				}
+				InstallerState = ChunkInstallState::Installing;
+			}
+			else
+			{
+				PriorityQueue.RemoveAt(0);
+			}
+		}
+		else
+		{
+			PriorityQueue.RemoveAt(0);
+		}
+	}
+	if (InstallingChunkID == -1)
+	{
+		// Install the first available chunk
+		for (auto It = RemoteManifests.CreateConstIterator(); It; ++It)
+		{
+			if (It)
+			{
+				IBuildManifestPtr ChunkManifest = It.Value();
+				auto ChunkIDField = ChunkManifest->GetCustomField("ChunkID");
+				if (ChunkIDField.IsValid())
+				{
+					InstallingChunkID = ChunkIDField->AsInteger();
+					InstallingChunkManifest = ChunkManifest;
+					auto PatchField = ChunkManifest->GetCustomField("bIsPatch");
+					bool bIsPatch = PatchField.IsValid() ? PatchField->AsString() == TEXT("true") : false;
+					auto ChunkFolderName = FString::Printf(TEXT("%s%d"), !bIsPatch ? TEXT("base") : TEXT("patch"), ChunkIDField->AsInteger());
+					auto ChunkInstallDir = FPaths::Combine(*InstallDir, *ChunkFolderName);
+					auto ChunkStageDir = FPaths::Combine(*StageDir, *ChunkFolderName);
+					BPSModule->SetCloudDirectory(CloudDir);
+					BPSModule->SetStagingDirectory(ChunkStageDir);
+					InstallService = BPSModule->StartBuildInstall(nullptr, ChunkManifest, ChunkInstallDir, FBuildPatchBoolManifestDelegate::CreateRaw(this, &FHTTPChunkInstall::OSSInstallComplete));
+					if (InstallSpeed == EChunkInstallSpeed::Paused && !InstallService->IsPaused())
+					{
+						InstallService->TogglePauseInstall();
+					}
+					InstallerState = ChunkInstallState::Installing;
+					return;
+				}
+			}
+		}
+	}
+}
+
+
+EChunkLocation::Type FHTTPChunkInstall::GetChunkLocation(uint32 ChunkID)
+{
+#if !UE_BUILD_SHIPPING
+	if(bDebugNoInstalledRequired)
+	{
+		return EChunkLocation::BestLocation;
+	}
+#endif
+
+	if (bFirstRun)
+	{
+		/** Still waiting on setup to finish, report that nothing is installed yet... */
+		return EChunkLocation::NotAvailable;
+	}
+	TArray<IBuildManifestPtr> FoundManifests;
+	RemoteManifests.MultiFind(ChunkID, FoundManifests);
+	if (FoundManifests.Num() > 0)
+	{
+		return EChunkLocation::NotAvailable;
+	}
+
+	InstalledManifests.MultiFind(ChunkID, FoundManifests);
+	if (FoundManifests.Num() > 0)
+	{
+		return EChunkLocation::BestLocation;
+	}
+
+	return EChunkLocation::DoesNotExist;
+}
+
+
+float FHTTPChunkInstall::GetChunkProgress(uint32 ChunkID,EChunkProgressReportingType::Type ReportType)
+{
+#if !UE_BUILD_SHIPPING
+	if (bDebugNoInstalledRequired)
+	{
+		return 100.f;
+	}
+#endif
+
+	if (bFirstRun)
+	{
+		/** Still waiting on setup to finish, report that nothing is installed yet... */
+		return 0.f;
+	}
+	TArray<IBuildManifestPtr> FoundManifests;
+	RemoteManifests.MultiFind(ChunkID, FoundManifests);
+	if (FoundManifests.Num() > 0)
+	{
+		float Progress = 0;
+		if (InstallingChunkID == ChunkID && InstallService.IsValid())
+		{
+			Progress = InstallService->GetUpdateProgress();
+		}
+		return Progress / FoundManifests.Num();
+	}
+
+	InstalledManifests.MultiFind(ChunkID, FoundManifests);
+	if (FoundManifests.Num() > 0)
+	{
+		return 100.f;
+	}
+
+	return 0.f;
+}
+
+void FHTTPChunkInstall::OSSEnumerateFilesComplete(bool bSuccess)
+{
+	InstallerState = bSuccess ? ChunkInstallState::SearchTitleFiles : ChunkInstallState::QueryRemoteManifests;
+}
+
+void FHTTPChunkInstall::OSSReadFileComplete(bool bSuccess, const FString& Filename)
+{
+	InstallerState = bSuccess ? ChunkInstallState::ReadComplete : ChunkInstallState::ReadTitleFiles;
+}
+
+void FHTTPChunkInstall::OSSInstallComplete(bool bSuccess, IBuildManifestRef BuildManifest)
+{
+	if (bSuccess)
+	{
+		// Completed OK. Write the manifest. If the chunk doesn't exist, copy to the content dir.
+		// Otherwise, writing the manifest will prompt a copy on next start of the game
+		FString ManifestName;
+		FString ChunkFdrName;
+		uint32 ChunkID;
+		bool bIsPatch;
+		if (!BuildChunkFolderName(BuildManifest, ChunkFdrName, ManifestName, ChunkID, bIsPatch))
+		{
+			//Something bad has happened, bail
+			InstallerState = ChunkInstallState::Idle;
+			return;
+		}
+		FString ManifestPath = FPaths::Combine(*InstallDir, *ChunkFdrName, *ManifestName);
+		FString SrcDir = FPaths::Combine(*InstallDir, *ChunkFdrName);
+		FString DestDir = FPaths::Combine(*ContentDir, *ChunkFdrName);
+		bool bCopyDir = true; 
+		TArray<IBuildManifestPtr> FoundManifests;
+		InstalledManifests.MultiFind(ChunkID, FoundManifests);
+		for (const auto& It : FoundManifests)
+		{
+			auto FoundPatchField = It->GetCustomField("bIsPatch");
+			bool bFoundPatch = FoundPatchField.IsValid() ? FoundPatchField->AsString() == TEXT("true") : false;
+			if (bFoundPatch == bIsPatch)
+			{
+				bCopyDir = false;
+			}
+		}
+		ChunkCopyInstall.GetTask().Init(ManifestPath, SrcDir, DestDir, BPSModule, BuildManifest, MountedPaks, bCopyDir);
+		ChunkCopyInstall.StartBackgroundTask();
+	}
+	InstallerState = ChunkInstallState::CopyToContent;
+}
+
+void FHTTPChunkInstall::ParseTitleFileManifest()
+{
+#if !UE_BUILD_SHIPPING
+	if (bDebugNoInstalledRequired)
+	{
+		// Forces the installer to think that no remote manifests exist, so nothing needs to be installed.
+		return;
+	}
+#endif
+	FString JsonBuffer;
+	FFileHelper::BufferToString(JsonBuffer, FileContentBuffer.GetData(), FileContentBuffer.Num());
+	auto RemoteManifest = BPSModule->MakeManifestFromJSON(JsonBuffer);
+	if (!RemoteManifest.IsValid())
+	{
+		return;
+	}
+	auto RemoteChunkIDField = RemoteManifest->GetCustomField("ChunkID");
+	if (!RemoteChunkIDField.IsValid())
+	{
+		return;
+	}
+	//Compare to installed manifests and add to the remote if it needs to be installed.
+	uint32 ChunkID = (uint32)RemoteChunkIDField->AsInteger();
+	TArray<IBuildManifestPtr> FoundManifests;
+	InstalledManifests.MultiFind(ChunkID, FoundManifests);
+	uint32 FoundCount = FoundManifests.Num();
+	if (FoundCount > 0)
+	{
+		auto RemotePatchManifest = RemoteManifest->GetCustomField("bIsPatch");
+		auto RemoteVersion = RemoteManifest->GetVersionString();
+		bool bRemoteIsPatch = RemotePatchManifest.IsValid() ? RemotePatchManifest->AsString() == TEXT("true") : false;
+		for (uint32 FoundIndex = 0; FoundIndex < FoundCount; ++FoundIndex)
+		{
+			const auto& InstalledManifest = FoundManifests[FoundIndex];
+			auto InstalledVersion = InstalledManifest->GetVersionString();
+			auto InstallPatchManifest = InstalledManifest->GetCustomField("bIsPatch");
+			bool bInstallIsPatch = InstallPatchManifest.IsValid() ? InstallPatchManifest->AsString() == TEXT("true") : false;
+			if (InstalledVersion != RemoteVersion && bInstallIsPatch == bRemoteIsPatch)
+			{
+				RemoteManifests.Add(ChunkID, RemoteManifest);
+				//Remove from the installed map
+				if (bFirstRun)
+				{
+					// Prevent the paks from being mounted by removing the manifest file
+					FString ChunkFdrName;
+					FString ManifestName;
+					uint32 ChunkID;
+					bool bIsPatch;
+					if (BuildChunkFolderName(InstalledManifest.ToSharedRef(), ChunkFdrName, ManifestName, ChunkID, bIsPatch))
+					{
+						FString ManifestPath = FPaths::Combine(*ContentDir, *ChunkFdrName, *ManifestName);
+						IPlatformFile& PlatformFile = FPlatformFileManager::Get().GetPlatformFile();
+						PlatformFile.DeleteFile(*ManifestPath);
+					}
+					InstalledManifests.Remove(ChunkID, InstalledManifest);
+				}
+			}
+		}
+	}
+	else
+	{
+		RemoteManifests.Add(ChunkID, RemoteManifest);
+	}
+}
+
+bool FHTTPChunkInstall::BuildChunkFolderName(IBuildManifestRef Manifest, FString& ChunkFdrName, FString& ManifestName, uint32& ChunkID, bool& bIsPatch)
+{
+	auto ChunkIDField = Manifest->GetCustomField("ChunkID");
+	auto ChunkPatchField = Manifest->GetCustomField("bIsPatch");
+
+	if (!ChunkIDField.IsValid())
+	{
+		return false;
+	}
+	ChunkID = ChunkIDField->AsInteger();
+	bIsPatch = ChunkPatchField.IsValid() ? ChunkPatchField->AsString() == TEXT("true") : false;
+	ManifestName = FString::Printf(TEXT("chunk_%u"), ChunkID);
+	if (bIsPatch)
+	{
+		ManifestName += TEXT("_patch");
+	}
+	ManifestName += TEXT(".manifest");
+	ChunkFdrName = FString::Printf(TEXT("%s%d"), !bIsPatch ? TEXT("base") : TEXT("patch"), ChunkID);
+	return true;
+}
+
+bool FHTTPChunkInstall::PrioritizeChunk(uint32 ChunkID, EChunkPriority::Type Priority)
+{	
+	int32 FoundIndex;
+	PriorityQueue.Find(FChunkPrio(ChunkID, Priority), FoundIndex);
+	if (FoundIndex != INDEX_NONE)
+	{
+		PriorityQueue.RemoveAt(FoundIndex);
+	}
+	// Low priority is assumed if the chunk ID doesn't exist in the queue
+	if (Priority != EChunkPriority::Low)
+	{
+		PriorityQueue.AddUnique(FChunkPrio(ChunkID, Priority));
+		PriorityQueue.Sort();
+	}
+	return true;
+}
+
+bool FHTTPChunkInstall::SetChunkInstallDelgate(uint32 ChunkID, FPlatformChunkInstallCompleteDelegate Delegate)
+{
+	FPlatformChunkInstallCompleteMultiDelegate* FoundDelegate = DelegateMap.Find(ChunkID);
+	if (FoundDelegate)
+	{
+		FoundDelegate->Add(Delegate);
+	}
+	else
+	{
+		FPlatformChunkInstallCompleteMultiDelegate MC;
+		MC.Add(Delegate);
+		DelegateMap.Add(ChunkID, MC);
+	}
+	return false;
+}
+
+void FHTTPChunkInstall::RemoveChunkInstallDelgate(uint32 ChunkID, FPlatformChunkInstallCompleteDelegate Delegate)
+{
+	FPlatformChunkInstallCompleteMultiDelegate* FoundDelegate = DelegateMap.Find(ChunkID);
+	if (!FoundDelegate)
+	{
+		return;
+	}
+	FoundDelegate->Remove(Delegate);
+}
+
+/**
+ * Module for the HTTP Chunk Installer
+ */
+class FHTTPChunkInstallerModule : public IPlatformChunkInstallModule
+{
+public:
+	TUniquePtr<IPlatformChunkInstall> ChunkInstaller;
+
+	FHTTPChunkInstallerModule()
+		: ChunkInstaller(new FHTTPChunkInstall())
+	{
+
+	}
+
+	virtual IPlatformChunkInstall* GetPlatformChunkInstall()
+	{
+		return ChunkInstaller.Get();
+	}
+};
+
+IMPLEMENT_MODULE(FHTTPChunkInstallerModule, HTTPChunkInstaller);
\ No newline at end of file
diff --git a/Engine/Source/Runtime/HTTPChunkInstaller/Private/HTTPChunkInstaller.h b/Engine/Source/Runtime/HTTPChunkInstaller/Private/HTTPChunkInstaller.h
new file mode 100644
--- /dev/null
+++ b/Engine/Source/Runtime/HTTPChunkInstaller/Private/HTTPChunkInstaller.h
@@ -0,0 +1,146 @@
+// Copyright 1998-2014 Epic Games, Inc. All Rights Reserved.
+
+#pragma once
+
+#include "GenericPlatformChunkInstall.h"
+#include "UniquePtr.h"
+#include "BuildPatchServices.h"
+#include "OnlineSubsystemTypes.h"
+#include "OnlineSubsystem.h"
+#include "ChunkInstall.h"
+#include "ChunkSetup.h"
+
+/**
+* HTTP based implementation of chunk based install
+**/
+class HTTPCHUNKINSTALLER_API FHTTPChunkInstall : public IPlatformChunkInstall, public FTickerObjectBase
+{
+public:
+	FHTTPChunkInstall();
+	~FHTTPChunkInstall();
+
+	virtual EChunkLocation::Type GetChunkLocation( uint32 ChunkID ) override;
+
+	virtual bool GetProgressReportingTypeSupported(EChunkProgressReportingType::Type ReportType) override
+	{
+		if (ReportType == EChunkProgressReportingType::PercentageComplete)
+		{
+			return true;
+		}
+
+		return false;
+	}
+
+	virtual float GetChunkProgress( uint32 ChunkID, EChunkProgressReportingType::Type ReportType ) override;
+
+	virtual EChunkInstallSpeed::Type GetInstallSpeed() override
+	{
+		return InstallSpeed;
+	}
+
+	virtual bool SetInstallSpeed( EChunkInstallSpeed::Type InInstallSpeed ) override
+	{
+		if (InstallSpeed != InInstallSpeed)
+		{
+			InstallSpeed = InInstallSpeed;
+			if ((InstallSpeed == EChunkInstallSpeed::Paused && !InstallService->IsPaused()) || (InstallSpeed != EChunkInstallSpeed::Paused && InstallService->IsPaused()))
+			{
+				InstallService->TogglePauseInstall();
+			}
+		}
+		return true;
+	}
+	
+	virtual bool PrioritizeChunk( uint32 ChunkID, EChunkPriority::Type Priority ) override;
+	virtual bool SetChunkInstallDelgate(uint32 ChunkID, FPlatformChunkInstallCompleteDelegate Delegate) override;
+	virtual void RemoveChunkInstallDelgate(uint32 ChunkID, FPlatformChunkInstallCompleteDelegate Delegate);
+
+	virtual bool DebugStartNextChunk()
+	{
+		/** Unless paused we are always installing! */
+		return false;
+	}
+
+	bool Tick(float DeltaSeconds) override;
+
+private:
+
+	void UpdatePendingInstallQueue();
+	void ParseTitleFileManifest();
+	bool BuildChunkFolderName(IBuildManifestRef Manifest, FString& ChunkFdrName, FString& ManifestName, uint32& ChunkID, bool& bIsPatch);
+	void OSSEnumerateFilesComplete(bool bSuccess);
+	void OSSReadFileComplete(bool bSuccess, const FString& Filename);
+	void OSSInstallComplete(bool, IBuildManifestRef);
+
+	DECLARE_MULTICAST_DELEGATE_OneParam(FPlatformChunkInstallCompleteMultiDelegate, uint32);
+
+	struct ChunkInstallState
+	{
+		enum Type
+		{
+			Setup,
+			SetupWait,
+			QueryRemoteManifests,
+			MoveInstalledChunks,
+			RequestingTitleFiles,
+			SearchTitleFiles,
+			ReadTitleFiles,
+			WaitingOnRead,
+			ReadComplete,
+			PostSetup,
+			Idle,
+			Installing,
+			CopyToContent,
+		};
+	};
+
+	struct FChunkPrio
+	{
+		FChunkPrio(uint32 InChunkID, EChunkPriority::Type InChunkPrio)
+			: ChunkID(InChunkID)
+			, ChunkPrio(InChunkPrio)
+		{
+
+		}
+
+		uint32					ChunkID;
+		EChunkPriority::Type	ChunkPrio;
+
+		bool operator == (const FChunkPrio& RHS) const
+		{
+			return ChunkID == RHS.ChunkID;
+		}
+
+		bool operator < (const FChunkPrio& RHS) const
+		{
+			return ChunkPrio < RHS.ChunkPrio;
+		}
+	};
+
+	FAsyncTask<FChunkInstallTask>								ChunkCopyInstall;
+	FAsyncTask<FChunkSetupTask>									ChunkSetupTask;
+	FAsyncTask<FChunkMountTask>									ChunkMountTask;
+	TMultiMap<uint32, IBuildManifestPtr>						InstalledManifests;
+	TMultiMap<uint32, IBuildManifestPtr>						RemoteManifests;
+	TMap<uint32, FPlatformChunkInstallCompleteMultiDelegate>	DelegateMap;
+	TArray<FCloudFileHeader>									TitleFilesToRead;
+	TArray<uint8>												FileContentBuffer;
+	TArray<FChunkPrio>											PriorityQueue;
+	TArray<FString>												MountedPaks;
+	IOnlineTitleFilePtr											OnlineTitleFile;
+	IBuildInstallerPtr											InstallService;
+	IBuildManifestPtr											InstallingChunkManifest;
+	FString														CloudDir;
+	FString														StageDir;
+	FString														InstallDir;
+	FString														BackupDir;
+	FString														ContentDir;
+	IBuildPatchServicesModule*									BPSModule;
+	uint32														InstallingChunkID;
+	ChunkInstallState::Type										InstallerState;
+	EChunkInstallSpeed::Type									InstallSpeed;
+	bool														bFirstRun;
+#if !UE_BUILD_SHIPPING
+	bool														bDebugNoInstalledRequired;
+#endif
+};
\ No newline at end of file
diff --git a/Engine/Source/Runtime/HTTPChunkInstaller/Private/HTTPChunkInstallerPrivatePCH.h b/Engine/Source/Runtime/HTTPChunkInstaller/Private/HTTPChunkInstallerPrivatePCH.h
new file mode 100644
--- /dev/null
+++ b/Engine/Source/Runtime/HTTPChunkInstaller/Private/HTTPChunkInstallerPrivatePCH.h
@@ -0,0 +1,7 @@
+// Copyright 1998-2014 Epic Games, Inc. All Rights Reserved.
+
+#include "Core.h"
+#include "Http.h"
+#include "ModuleManager.h"
+#include "Engine.h"
+#include "OnlineSubsystemUtils.h"
diff --git a/Engine/Source/Runtime/HTTPChunkInstaller/Private/LocalTitleFile.cpp b/Engine/Source/Runtime/HTTPChunkInstaller/Private/LocalTitleFile.cpp
new file mode 100644
--- /dev/null
+++ b/Engine/Source/Runtime/HTTPChunkInstaller/Private/LocalTitleFile.cpp
@@ -0,0 +1,132 @@
+// Copyright 1998-2014 Epic Games, Inc. All Rights Reserved.
+
+#include "HTTPChunkInstallerPrivatePCH.h"
+#include "LocalTitleFile.h"
+
+
+FLocalTitleFile::FLocalTitleFile(const FString& InRootDirectory)
+: RootDirectory(InRootDirectory)
+{ }
+
+
+bool FLocalTitleFile::GetFileContents(const FString& DLName, TArray<uint8>& FileContents)
+{
+	const TArray< uint8 >* ExistingFileContents = DLNameToFileContents.Find(DLName);
+
+	if (ExistingFileContents == nullptr)
+	{
+		return false;
+	}
+
+	FileContents.Append(*ExistingFileContents);
+
+	return true;
+}
+
+
+bool FLocalTitleFile::ClearFiles()
+{
+	FileHeaders.Empty();
+	DLNameToFileContents.Empty();
+
+	return true;
+}
+
+
+bool FLocalTitleFile::ClearFile(const FString& DLName)
+{
+	bool ClearedFile = false;
+	const FString FileName = GetFileNameFromDLName(DLName);
+
+	for (int Index = 0; Index < FileHeaders.Num(); Index++)
+	{
+		if (FileHeaders[Index].DLName == DLName)
+		{
+			FileHeaders.RemoveAt(Index);
+			ClearedFile = true;
+		}
+	}
+
+	DLNameToFileContents.Remove(DLName);
+
+	return ClearedFile;
+}
+
+
+bool FLocalTitleFile::EnumerateFiles(const FPagedQuery& Page)
+{
+	const FString WildCard = FPaths::Combine(*RootDirectory, TEXT("*"));
+
+	TArray<FString> Filenames;
+	IFileManager::Get().FindFiles(Filenames, *WildCard, true, false);
+
+	bool Success = Filenames.Num() > 0;
+
+	if (Success)
+	{
+		for (int32 FileIdx = 0; FileIdx < Filenames.Num(); ++FileIdx)
+		{
+			const FString Filename = Filenames[FileIdx];
+
+			FCloudFileHeader NewHeader;
+			NewHeader.FileName = Filename;
+			NewHeader.DLName = Filename + FileIdx;
+			NewHeader.FileSize = 0;
+			NewHeader.Hash.Empty();
+
+			FileHeaders.Add(NewHeader);
+		}
+	}
+
+	TriggerOnEnumerateFilesCompleteDelegates(Success);
+
+	return Success;
+}
+
+void FLocalTitleFile::GetFileList(TArray<FCloudFileHeader>& InFileHeaders)
+{
+	InFileHeaders.Append(FileHeaders);
+}
+
+
+bool FLocalTitleFile::ReadFile(const FString& DLName)
+{
+	const TArray< uint8 >* ExistingFileContents = DLNameToFileContents.Find(DLName);
+
+	if (ExistingFileContents != nullptr)
+	{
+		TriggerOnReadFileCompleteDelegates(true, DLName);
+
+		return true;
+	}
+
+	const FString FileName = GetFileNameFromDLName(DLName);
+
+	TArray<uint8> FileContents;
+
+	if (!FFileHelper::LoadFileToArray(FileContents, *(RootDirectory + FileName)))
+	{
+		TriggerOnReadFileCompleteDelegates(false, DLName);
+
+		return false;
+	}
+
+	DLNameToFileContents.Add(DLName, FileContents);
+	TriggerOnReadFileCompleteDelegates(true, DLName);
+
+	return true;
+}
+
+
+FString FLocalTitleFile::GetFileNameFromDLName(const FString& DLName) const
+{
+	for (int Index = 0; Index < FileHeaders.Num(); Index++)
+	{
+		if (FileHeaders[Index].DLName == DLName)
+		{
+			return FileHeaders[Index].FileName;
+		}
+	}
+
+	return FString();
+}
\ No newline at end of file
diff --git a/Engine/Source/Runtime/HTTPChunkInstaller/Private/LocalTitleFile.h b/Engine/Source/Runtime/HTTPChunkInstaller/Private/LocalTitleFile.h
new file mode 100644
--- /dev/null
+++ b/Engine/Source/Runtime/HTTPChunkInstaller/Private/LocalTitleFile.h
@@ -0,0 +1,31 @@
+// Copyright 1998-2014 Epic Games, Inc. All Rights Reserved.
+
+#pragma once
+
+#include "OnlineTitleFileInterface.h"
+
+class FLocalTitleFile : public IOnlineTitleFile
+{
+public:
+
+	FLocalTitleFile(const FString& InRootDirectory);
+
+	// IOnlineTitleFile interface
+
+	virtual bool GetFileContents(const FString& DLName, TArray<uint8>& FileContents) override;
+	virtual bool ClearFiles() override;
+	virtual bool ClearFile(const FString& DLName) override;
+	virtual bool EnumerateFiles(const FPagedQuery& Page = FPagedQuery()) override;
+	virtual void GetFileList(TArray<FCloudFileHeader>& InFileHeaders) override;
+	virtual bool ReadFile(const FString& DLName) override;
+
+protected:
+
+	FString GetFileNameFromDLName(const FString& DLName) const;
+
+private:
+
+	FString							RootDirectory;
+	TArray<FCloudFileHeader>		FileHeaders;
+	TMap< FString, TArray<uint8> >	DLNameToFileContents;
+};
\ No newline at end of file
diff --git a/Engine/Source/Runtime/PakFile/Private/IPlatformFilePak.cpp b/Engine/Source/Runtime/PakFile/Private/IPlatformFilePak.cpp
--- a/Engine/Source/Runtime/PakFile/Private/IPlatformFilePak.cpp
+++ b/Engine/Source/Runtime/PakFile/Private/IPlatformFilePak.cpp
@@ -462,17 +462,17 @@ void FPakPlatformFile::HandleMountCommand(const TCHAR* Cmd, FOutputDevice& Ar)
 	if (!PakFilename.IsEmpty())
 	{
 		const FString MountPoint = FParse::Token(Cmd, false);
-		Mount(*PakFilename, MountPoint.IsEmpty() ? NULL : *MountPoint);
+		Mount(*PakFilename, 0, MountPoint.IsEmpty() ? NULL : *MountPoint);
 	}
 }
 
 void FPakPlatformFile::HandlePakListCommand(const TCHAR* Cmd, FOutputDevice& Ar)
 {
-	TArray<FPakFile*> Paks;
+	TArray<FPakListEntry> Paks;
 	GetMountedPaks(Paks);
 	for (auto Pak : Paks)
 	{
-		Ar.Logf(TEXT("%s"), *Pak->GetFilename());
+		Ar.Logf(TEXT("%s"), *Pak.PakFile->GetFilename());
 	}	
 }
 #endif // !UE_BUILD_SHIPPING
@@ -498,7 +498,8 @@ FPakPlatformFile::~FPakPlatformFile()
 		FScopeLock ScopedLock(&PakListCritical);
 		for (int32 PakFileIndex = 0; PakFileIndex < PakFiles.Num(); PakFileIndex++)
 		{
-			delete PakFiles[PakFileIndex];
+			delete PakFiles[PakFileIndex].PakFile;
+			PakFiles[PakFileIndex].PakFile = nullptr;
 		}
 	}	
 }
@@ -631,7 +632,7 @@ bool FPakPlatformFile::Initialize(IPlatformFile* Inner, const TCHAR* CmdLine)
 		}
 		if (bLoadPak)
 		{
-			Mount(*PakFilename);
+			Mount(*PakFilename, 0);
 		}
 	}
 
@@ -643,7 +644,7 @@ bool FPakPlatformFile::Initialize(IPlatformFile* Inner, const TCHAR* CmdLine)
 	return !!LowerLevel;
 }
 
-bool FPakPlatformFile::Mount(const TCHAR* InPakFilename, const TCHAR* InPath /*= NULL*/)
+bool FPakPlatformFile::Mount(const TCHAR* InPakFilename, uint32 PakOrder, const TCHAR* InPath /*= NULL*/)
 {
 	bool bSuccess = false;
 	IFileHandle* PakHandle = LowerLevel->OpenRead(InPakFilename);
@@ -659,7 +660,11 @@ bool FPakPlatformFile::Mount(const TCHAR* InPakFilename, const TCHAR* InPath /*=
 			{
 				// Add new pak file
 				FScopeLock ScopedLock(&PakListCritical);
-				PakFiles.Add(Pak);
+				FPakListEntry Entry;
+				Entry.ReadOrder = PakOrder;
+				Entry.PakFile = Pak;
+				PakFiles.Add(Entry);
+				PakFiles.Sort();
 			}
 			bSuccess = true;
 		}
@@ -704,9 +709,9 @@ IFileHandle* FPakPlatformFile::CreatePakFileHandle(const TCHAR* Filename, FPakFi
 	return Result;
 }
 
-bool FPakPlatformFile::HandleMountPakDelegate(const FString& PakFilePath)
+bool FPakPlatformFile::HandleMountPakDelegate(const FString& PakFilePath, uint32 PakOrder)
 {
-	return Mount(*PakFilePath);
+	return Mount(*PakFilePath, PakOrder);
 }
 
 IFileHandle* FPakPlatformFile::OpenRead(const TCHAR* Filename)
diff --git a/Engine/Source/Runtime/PakFile/Public/IPlatformFilePak.h b/Engine/Source/Runtime/PakFile/Public/IPlatformFilePak.h
--- a/Engine/Source/Runtime/PakFile/Public/IPlatformFilePak.h
+++ b/Engine/Source/Runtime/PakFile/Public/IPlatformFilePak.h
@@ -774,10 +774,26 @@ public:
  **/
 class PAKFILE_API FPakPlatformFile : public IPlatformFile
 {
+	struct FPakListEntry
+	{
+		FPakListEntry()
+			: ReadOrder(0)
+			, PakFile(nullptr)
+		{}
+
+		uint32		ReadOrder;
+		FPakFile*	PakFile;
+
+		bool operator < (const FPakListEntry& RHS) const
+		{
+			return ReadOrder > RHS.ReadOrder;
+		}
+	};
+
 	/** Wrapped file */
 	IPlatformFile* LowerLevel;
 	/** List of all available pak files. */
-	TArray<FPakFile*> PakFiles;
+	TArray<FPakListEntry> PakFiles;
 	/** True if this we're using signed content. */
 	bool bSigned;
 	/** Synchronization object for accessing the list of currently mounted pak files. */
@@ -786,7 +802,7 @@ class PAKFILE_API FPakPlatformFile : public IPlatformFile
 	/**
 	 * Gets mounted pak files
 	 */
-	FORCEINLINE void GetMountedPaks(TArray<FPakFile*>& Paks)
+	FORCEINLINE void GetMountedPaks(TArray<FPakListEntry>& Paks)
 	{
 		FScopeLock ScopedLock(&PakListCritical);
 		Paks.Append(PakFiles);
@@ -803,13 +819,13 @@ class PAKFILE_API FPakPlatformFile : public IPlatformFile
 		FString StandardPath = Directory;
 		FPaths::MakeStandardFilename(StandardPath);
 
-		TArray<FPakFile*> Paks;
+		TArray<FPakListEntry> Paks;
 		GetMountedPaks(Paks);
 
 		// Check all pak files.
 		for (int32 PakIndex = 0; PakIndex < Paks.Num(); PakIndex++)
 		{
-			if (Paks[PakIndex]->DirectoryExists(*StandardPath))
+			if (Paks[PakIndex].PakFile->DirectoryExists(*StandardPath))
 			{
 				return true;
 			}
@@ -842,7 +858,7 @@ class PAKFILE_API FPakPlatformFile : public IPlatformFile
 	/**
 	 * Handler for device delegate to prompt us to load a new pak.	 
 	 */
-	bool HandleMountPakDelegate(const FString& PakFilePath);
+	bool HandleMountPakDelegate(const FString& PakFilePath, uint32 PakOrder);
 
 	/**
 	 * Finds all pak files in the given directory.
@@ -907,7 +923,7 @@ public:
 	 * @param InPakFilename Pak filename.
 	 * @param InPath Path to mount the pak at.
 	 */
-	bool Mount(const TCHAR* InPakFilename, const TCHAR* InPath = NULL);
+	bool Mount(const TCHAR* InPakFilename, uint32 PakOrder, const TCHAR* InPath = NULL);
 
 	/**
 	 * Finds a file in the specified pak files.
@@ -917,7 +933,7 @@ public:
 	 * @param OutPakFile Optional pointer to a pak file where the filename was found.
 	 * @return Pointer to pak entry if the file was found, NULL otherwise.
 	 */
-	FORCEINLINE static const FPakEntry* FindFileInPakFiles(TArray<FPakFile*>& Paks, const TCHAR* Filename,  FPakFile** OutPakFile)
+	FORCEINLINE static const FPakEntry* FindFileInPakFiles(TArray<FPakListEntry>& Paks,const TCHAR* Filename,FPakFile** OutPakFile)
 	{
 		FString StandardFilename(Filename);
 		FPaths::MakeStandardFilename(StandardFilename);
@@ -925,12 +941,12 @@ public:
 
 		for (int32 PakIndex = 0; !FoundEntry && PakIndex < Paks.Num(); PakIndex++)
 		{
-			FoundEntry = Paks[PakIndex]->Find(*StandardFilename);
+			FoundEntry = Paks[PakIndex].PakFile->Find(*StandardFilename);
 			if (FoundEntry != NULL)
 			{
 				if (OutPakFile != NULL)
 				{
-					*OutPakFile = Paks[PakIndex];
+					*OutPakFile = Paks[PakIndex].PakFile;
 				}
 			}
 		}
@@ -947,7 +963,7 @@ public:
 	 */
 	const FPakEntry* FindFileInPakFiles(const TCHAR* Filename, FPakFile** OutPakFile = NULL)
 	{
-		TArray<FPakFile*> Paks;
+		TArray<FPakListEntry> Paks;
 		GetMountedPaks(Paks);
 
 		return FindFileInPakFiles(Paks, Filename, OutPakFile);
@@ -1118,10 +1134,10 @@ public:
 		/** Visited pak files. */
 		TSet<FString>& VisitedPakFiles;
 		/** Cached list of pak files. */
-		TArray<FPakFile*>& Paks;
+		TArray<FPakListEntry>& Paks;
 
 		/** Constructor. */
-		FPakVisitor(FDirectoryVisitor& InVisitor, TArray<FPakFile*>& InPaks, TSet<FString>& InVisitedPakFiles)
+		FPakVisitor(FDirectoryVisitor& InVisitor, TArray<FPakListEntry>& InPaks, TSet<FString>& InVisitedPakFiles)
 			: Visitor(InVisitor)
 			, VisitedPakFiles(InVisitedPakFiles)
 			, Paks(InPaks)
@@ -1152,13 +1168,13 @@ public:
 		bool Result = true;
 		TSet<FString> FilesVisitedInPak;
 
-		TArray<FPakFile*> Paks;
+		TArray<FPakListEntry> Paks;
 		GetMountedPaks(Paks);
 
 		// Iterate pak files first
 		for (int32 PakIndex = 0; PakIndex < Paks.Num(); PakIndex++)
 		{
-			FPakFile& PakFile = *Paks[PakIndex];
+			FPakFile& PakFile = *Paks[PakIndex].PakFile;
 			
 			const bool bIncludeFiles = true;
 			const bool bIncludeFolders = true;
@@ -1205,7 +1221,7 @@ public:
 	virtual bool IterateDirectoryRecursively(const TCHAR* Directory, IPlatformFile::FDirectoryVisitor& Visitor) override
 	{
 		TSet<FString> FilesVisitedInPak;
-		TArray<FPakFile*> Paks;
+		TArray<FPakListEntry> Paks;
 		GetMountedPaks(Paks);
 		FPakVisitor PakVisitor(Visitor, Paks, FilesVisitedInPak);
 		return IPlatformFile::IterateDirectoryRecursively(Directory, PakVisitor);
