HTTPChunkInstall Improvements and Bug fixes.
Added -chunkinstalldirectory and -chunkinstallversion to UAT commandline options.
Added Chunk Dependency Graph for use with -manifests.
HTTPChunkInstaller correctly re-uses old builds to reduce download sizes.
ChunkInstaller delegates correctly use delegate handles.
Windows uses HTTPChunkInstaller by default if enabled in project Build.cs.
Switched HTTPChunkInstaller tasks to threads.

diff --git a/Engine/Source/Editor/UnrealEd/Private/Commandlets/ChunkDependencyInfo.cpp b/Engine/Source/Editor/UnrealEd/Private/Commandlets/ChunkDependencyInfo.cpp
new file mode 100644
--- /dev/null
+++ b/Engine/Source/Editor/UnrealEd/Private/Commandlets/ChunkDependencyInfo.cpp
@@ -0,0 +1,65 @@
+// Copyright 1998-2015 Epic Games, Inc. All Rights Reserved.
+
+#include "UnrealEd.h"
+#include "ChunkDependencyInfo.h"
+
+UChunkDependencyInfo::UChunkDependencyInfo(const FObjectInitializer& ObjectInitializer)
+	: Super(ObjectInitializer)
+{
+}
+
+const FChunkDependencyTreeNode* UChunkDependencyInfo::GetChunkDependencyGraph(uint32 TotalChunks)
+{
+	// Reset any current tree
+	RootTreeNode.ChunkID = 0;
+	RootTreeNode.ChildNodes.Reset(0);
+
+	// Ensure the DependencyArray is OK to work with.
+	// Remove cycles
+	DependencyArray.RemoveAllSwap([](const FChunkDependency& A){ return A.ChunkID == A.ParentChunkID; });
+	// Add missing links (assumes they parent to chunk zero)
+	for (uint32 i = 1; i < TotalChunks; ++i)
+	{
+		if (!DependencyArray.FindByPredicate([=](const FChunkDependency& RHS){ return i == RHS.ChunkID; }))
+		{
+			FChunkDependency Dep;
+			Dep.ChunkID = i;
+			Dep.ParentChunkID = 0;
+			DependencyArray.Add(Dep);
+		}
+	}
+	// Remove duplicates
+	DependencyArray.StableSort([](const FChunkDependency& LHS, const FChunkDependency& RHS) { return LHS.ChunkID < RHS.ChunkID; });
+	for (int32 i = 0; i < DependencyArray.Num() - 1;)
+	{
+		if (DependencyArray[i] == DependencyArray[i + 1])
+		{
+			DependencyArray.RemoveAt(i + 1);
+		}
+		else
+		{
+			++i;
+		}
+	}
+
+	//
+	AddChildrenRecursive(RootTreeNode, DependencyArray);
+	return &RootTreeNode;
+}
+
+void UChunkDependencyInfo::AddChildrenRecursive(FChunkDependencyTreeNode& Node, TArray<FChunkDependency>& DepInfo)
+{
+	auto ChildNodeIndices = DepInfo.FilterByPredicate(
+		[&](const FChunkDependency& RHS) 
+		{
+			return Node.ChunkID == RHS.ParentChunkID;
+		});
+	for (const auto& ChildIndex : ChildNodeIndices)
+	{
+		Node.ChildNodes.Add(FChunkDependencyTreeNode(ChildIndex.ChunkID));
+	}
+	for (auto& Child : Node.ChildNodes) 
+	{
+		AddChildrenRecursive(Child, DepInfo);
+	}
+}
diff --git a/Engine/Source/Editor/UnrealEd/Private/Commandlets/ChunkDependencyInfo.h b/Engine/Source/Editor/UnrealEd/Private/Commandlets/ChunkDependencyInfo.h
new file mode 100644
--- /dev/null
+++ b/Engine/Source/Editor/UnrealEd/Private/Commandlets/ChunkDependencyInfo.h
@@ -0,0 +1,50 @@
+// Copyright 1998-2015 Epic Games, Inc. All Rights Reserved.
+
+#pragma once
+
+#include "ChunkDependencyInfo.generated.h"
+
+USTRUCT()
+struct FChunkDependency
+{
+	GENERATED_USTRUCT_BODY()
+
+	UPROPERTY(EditAnywhere, Category = ChunkInfo)
+	int32 ChunkID;
+
+	UPROPERTY(EditAnywhere, Category = ChunkInfo)
+	int32 ParentChunkID;
+
+	bool operator == (const FChunkDependency& RHS) 
+	{
+		return ChunkID == RHS.ChunkID;
+	}
+};
+
+struct FChunkDependencyTreeNode
+{
+	FChunkDependencyTreeNode(int32 InChunkID=0)
+		: ChunkID(InChunkID)
+	{}
+
+	int32 ChunkID;
+
+	TArray<FChunkDependencyTreeNode> ChildNodes;
+};
+
+UCLASS(config=Engine, defaultconfig)
+class UChunkDependencyInfo : public UObject
+{
+	GENERATED_UCLASS_BODY()
+
+	const FChunkDependencyTreeNode* GetChunkDependencyGraph(uint32 TotalChunks);
+
+public:
+
+	void AddChildrenRecursive(FChunkDependencyTreeNode& Node, TArray<FChunkDependency>& DepInfo);
+
+	UPROPERTY(config)
+	TArray<FChunkDependency> DependencyArray;
+
+	FChunkDependencyTreeNode RootTreeNode;
+};
\ No newline at end of file
diff --git a/Engine/Source/Editor/UnrealEd/Private/Commandlets/ChunkManifestGenerator.cpp b/Engine/Source/Editor/UnrealEd/Private/Commandlets/ChunkManifestGenerator.cpp
--- a/Engine/Source/Editor/UnrealEd/Private/Commandlets/ChunkManifestGenerator.cpp
+++ b/Engine/Source/Editor/UnrealEd/Private/Commandlets/ChunkManifestGenerator.cpp
@@ -12,6 +12,7 @@
 #include "UnrealEdMessages.h"
 #include "GameDelegates.h"
 #include "ChunkManifestGenerator.h"
+#include "ChunkDependencyInfo.h"
 #include "IPlatformFileSandboxWrapper.h"
 
 #include "JsonWriter.h"
@@ -25,6 +26,7 @@ FChunkManifestGenerator::FChunkManifestGenerator(const TArray<ITargetPlatform*>&
 	, Platforms(InPlatforms)
 	, bGenerateChunks(false)
 {
+	DependencyInfo = GetMutableDefault<UChunkDependencyInfo>();
 }
 
 FChunkManifestGenerator::~FChunkManifestGenerator()
@@ -726,6 +728,59 @@ void FChunkManifestGenerator::NotifyPackageWasCooked(const FString& PackageSandb
 	AllCookedPackages.Add(PackageName, PackageSandboxPath);
 }
 
+void FChunkManifestGenerator::ResolveChunkDependencyGraph(const FChunkDependencyTreeNode& Node, FChunkPackageSet BaseAssetSet) 
+{
+	if (FinalChunkManifests.Num() > Node.ChunkID && FinalChunkManifests[Node.ChunkID])
+	{
+		for (auto It = BaseAssetSet.CreateConstIterator(); It; ++It)
+		{
+			// Remove any assets belonging to our parents.
+			FinalChunkManifests[Node.ChunkID]->Remove(It.Key());
+		}
+		// Add the current Chunk's assets
+		for (auto It = FinalChunkManifests[Node.ChunkID]->CreateConstIterator(); It; ++It)//for (const auto It : *(FinalChunkManifests[Node.ChunkID]))
+		{
+			BaseAssetSet.Add(It.Key(), It.Value());
+		}
+		for (const auto It : Node.ChildNodes)
+		{
+			ResolveChunkDependencyGraph(It, BaseAssetSet);
+		}
+	}
+}
+
+bool FChunkManifestGenerator::CheckChunkAssetsAreNotInChild(const FChunkDependencyTreeNode& Node)
+{
+	for (const auto It : Node.ChildNodes)
+	{
+		if (!CheckChunkAssetsAreNotInChild(It))
+		{
+			return false;
+		}
+	}
+
+	if (!(FinalChunkManifests.Num() > Node.ChunkID && FinalChunkManifests[Node.ChunkID]))
+	{
+		return true;
+	}
+
+	for (const auto ChildIt : Node.ChildNodes)
+	{
+		if(FinalChunkManifests.Num() > ChildIt.ChunkID && FinalChunkManifests[ChildIt.ChunkID])
+		{
+			for (auto It = FinalChunkManifests[Node.ChunkID]->CreateConstIterator(); It; ++It)
+			{
+				if (FinalChunkManifests[ChildIt.ChunkID]->Find(It.Key()))
+				{
+					return false;
+				}
+			}
+		}
+	}
+
+	return true;
+}
+
 void FChunkManifestGenerator::FixupPackageDependenciesForChunks(FSandboxPlatformFile* InSandboxFile)
 {
 	for (int32 ChunkID = 0, MaxChunk = ChunkManifests.Num(); ChunkID < MaxChunk; ++ChunkID)
@@ -742,6 +797,7 @@ void FChunkManifestGenerator::FixupPackageDependenciesForChunks(FSandboxPlatform
 		}
 	}
 
+	auto* ChunkDepGraph = DependencyInfo->GetChunkDependencyGraph(ChunkManifests.Num());
 	//Once complete, Add any remaining assets (that are not assigned to a chunk) to the first chunk.
 	if (FinalChunkManifests.Num() == 0)
 	{
@@ -758,17 +814,17 @@ void FChunkManifestGenerator::FixupPackageDependenciesForChunks(FSandboxPlatform
 		AddPackageAndDependenciesToChunk(FinalChunkManifests[0], It.Key(), It.Value(), 0, InSandboxFile);
 	}
 
+	if (!CheckChunkAssetsAreNotInChild(*ChunkDepGraph))
+	{
+		UE_LOG(LogChunkManifestGenerator, Log, TEXT("Initial scan of chunks found duplicate assets in graph children"));
+	}
+
 	//Finally, if the previous step may added any extra packages to the 0 chunk. Pull them out of other chunks and save space
-	for (auto It = FinalChunkManifests[0]->CreateConstIterator(); It; ++It)
+	ResolveChunkDependencyGraph(*ChunkDepGraph, FChunkPackageSet());
+
+	if (!CheckChunkAssetsAreNotInChild(*ChunkDepGraph))
 	{
-		for (int32 ChunkID = 1, MaxChunk = FinalChunkManifests.Num(); ChunkID < MaxChunk; ++ChunkID)
-		{
-			if (!FinalChunkManifests[ChunkID])
-			{
-				continue;
-			}
-			FinalChunkManifests[ChunkID]->Remove(It.Key());
-		}
+		UE_LOG(LogChunkManifestGenerator, Error, TEXT("Second Scan of chunks found duplicate asset entries in children."));
 	}
 
 	// Fix up the asset registry to reflect this chunk layout
diff --git a/Engine/Source/Editor/UnrealEd/Private/Commandlets/ChunkManifestGenerator.h b/Engine/Source/Editor/UnrealEd/Private/Commandlets/ChunkManifestGenerator.h
--- a/Engine/Source/Editor/UnrealEd/Private/Commandlets/ChunkManifestGenerator.h
+++ b/Engine/Source/Editor/UnrealEd/Private/Commandlets/ChunkManifestGenerator.h
@@ -2,6 +2,9 @@
 
 #pragma once
 
+class UChunkDependencyInfo;
+struct FChunkDependencyTreeNode;
+
 /**
  * Helper class for generating streaming install manifests
  */
@@ -33,6 +36,8 @@ class FChunkManifestGenerator
 	TArray<FChunkPackageSet*>		FinalChunkManifests;
 	/** Lookup table of used package names used when searching references. */
 	TSet<FName>						InspectedNames;
+	/** */
+	UChunkDependencyInfo*			DependencyInfo;
 
 	struct FReferencePair
 	{
@@ -195,6 +200,16 @@ class FChunkManifestGenerator
 	*/
 	FString			GetShortestReferenceChain(FName PackageName, int32 ChunkID);
 
+	/**
+	* 
+	*/
+	void			ResolveChunkDependencyGraph(const FChunkDependencyTreeNode& Node, FChunkPackageSet BaseAssetSet);
+
+	/**
+	* Helper function to verify Chunk asset assigment is valid.
+	*/
+	bool			CheckChunkAssetsAreNotInChild(const FChunkDependencyTreeNode& Node);
+
 public:
 
 	/**
diff --git a/Engine/Source/Programs/AutomationTool/ProjectParams.cs b/Engine/Source/Programs/AutomationTool/ProjectParams.cs
--- a/Engine/Source/Programs/AutomationTool/ProjectParams.cs
+++ b/Engine/Source/Programs/AutomationTool/ProjectParams.cs
@@ -494,6 +494,8 @@ namespace AutomationTool
             this.StageNonMonolithic = GetParamValueIfNotSpecified(Command, StageNonMonolithic, this.StageNonMonolithic, "StageNonMonolithic");
 			this.Manifests = GetParamValueIfNotSpecified(Command, Manifests, this.Manifests, "manifests");
             this.CreateChunkInstall = GetParamValueIfNotSpecified(Command, CreateChunkInstall, this.CreateChunkInstall, "createchunkinstall");
+			this.ChunkInstallDirectory = ParseParamValueIfNotSpecified(Command, ChunkInstallDirectory, "chunkinstalldirectory", String.Empty).Trim(new char[]{'\"'});
+			this.ChunkInstallVersionString = ParseParamValueIfNotSpecified(Command, ChunkInstallVersionString, "chunkinstallversion", String.Empty).Trim(new char[] { '\"' });
 			this.Archive = GetParamValueIfNotSpecified(Command, Archive, this.Archive, "archive");
 			this.ArchiveDirectoryParam = ParseParamValueIfNotSpecified(Command, ArchiveDirectoryParam, "archivedirectory", String.Empty);
 			this.ArchiveMetaData = GetParamValueIfNotSpecified(Command, ArchiveMetaData, this.ArchiveMetaData, "archivemetadata");
@@ -813,11 +815,21 @@ namespace AutomationTool
 		public bool Manifests { private set; get; }
 
         /// <summary>
-        /// Shared: true if this build chunk install streaming install data, command line: -createchunkinstalldata !!JM
+        /// Shared: true if this build chunk install streaming install data, command line: -createchunkinstalldata
         /// </summary>
         [Help("createchunkinstall", "generate streaming install data from manifest when cooking data, requires -stage & -manifests")]
         public bool CreateChunkInstall { private set; get; }
 
+		/// <summary>
+		/// Shared: Directory to use for built chunk install data, command line: -chunkinstalldirectory=
+		/// </summary>
+		public string ChunkInstallDirectory { set; get; }
+
+		/// <summary>
+		/// Shared: Version string to use for built chunk install data, command line: -chunkinstallversion=
+		/// </summary>
+		public string ChunkInstallVersionString { set; get; }
+
 		/// <summary>
 		/// Shared: Directory to copy the client to, command line: -stagingdirectory=
 		/// </summary>	
@@ -1937,8 +1949,18 @@ namespace AutomationTool
 
             if (CreateChunkInstall && (!Manifests || !Stage))
             {
-                throw new AutomationException("-createchunkinstall can only be used with -pak & -stage");
+                throw new AutomationException("-createchunkinstall can only be used with -manifests & -stage"); 
             }
+
+			if (CreateChunkInstall && String.IsNullOrEmpty(ChunkInstallDirectory))
+			{
+				throw new AutomationException("-createchunkinstall must specify the chunk install data directory with -chunkinstalldirectory=");
+			}
+
+			if (CreateChunkInstall && String.IsNullOrEmpty(ChunkInstallVersionString))
+			{
+				throw new AutomationException("-createchunkinstall must specify the chunk install data version string with -chunkinstallversion=");
+			}
 		}
 
 		protected bool bLogged = false;
diff --git a/Engine/Source/Programs/AutomationTool/Scripts/CopyBuildToStagingDirectory.Automation.cs b/Engine/Source/Programs/AutomationTool/Scripts/CopyBuildToStagingDirectory.Automation.cs
--- a/Engine/Source/Programs/AutomationTool/Scripts/CopyBuildToStagingDirectory.Automation.cs
+++ b/Engine/Source/Programs/AutomationTool/Scripts/CopyBuildToStagingDirectory.Automation.cs
@@ -744,15 +744,9 @@ public partial class Project : CommandUtils
 				var BPTExe = GetBuildPatchToolExecutable();
 				EnsureBuildPatchToolExists();
 
-				string P4Change = "UnknownCL";
-				string P4Branch = "UnknownBranch";
-				if (CommandUtils.P4Enabled)
-				{
-					P4Change = CommandUtils.P4Env.ChangelistString;
-					P4Branch = CommandUtils.P4Env.BuildRootEscaped;
-				}
-				string ChunkInstallBasePath = CombinePaths(SC.ProjectRoot, "ChunkInstall", SC.FinalCookPlatform);
-				string RawDataPath = CombinePaths(ChunkInstallBasePath, P4Branch + "-CL-" + P4Change, PakName);
+				string VersionString = Params.ChunkInstallVersionString;
+				string ChunkInstallBasePath = CombinePaths(Params.ChunkInstallDirectory, SC.FinalCookPlatform);
+				string RawDataPath = CombinePaths(ChunkInstallBasePath, VersionString, PakName);
 				string RawDataPakPath = CombinePaths(RawDataPath, PakName + "-" + SC.FinalCookPlatform + ".pak");
 				//copy the pak chunk to the raw data folder
 				if (InternalUtils.SafeFileExists(RawDataPakPath, true))
@@ -768,10 +762,15 @@ public partial class Project : CommandUtils
 
 				string BuildRoot = MakePathSafeToUseWithCommandLine(RawDataPath);
 				string CloudDir = MakePathSafeToUseWithCommandLine(CombinePaths(ChunkInstallBasePath, "CloudDir"));
-				var AppID = 1; // !!JM todo, get valid value for this
+				string ManifestDir = CombinePaths(ChunkInstallBasePath, "ManifestDir");
+				var AppID = 1; // For a chunk install this value doesn't seem to matter
 				string AppName = String.Format("{0}_{1}", SC.ShortProjectName, PakName);
-				string AppLaunch = ""; // !!JM todo, get real value for this...?
-				string VersionString = P4Branch + "-CL-" + P4Change;
+				string AppLaunch = ""; // For a chunk install this value doesn't seem to matter
+				string ManifestFilename = AppName + VersionString + ".manifest";
+				string SourceManifestPath = CombinePaths(CloudDir, ManifestFilename);
+				string BackupManifestPath = CombinePaths(RawDataPath, ManifestFilename);
+				string DestManifestPath = CombinePaths(ManifestDir, ManifestFilename);
+				InternalUtils.SafeCreateDirectory(ManifestDir, true);
 
 				string CmdLine = String.Format("-BuildRoot=\"{0}\" -CloudDir=\"{1}\" -AppID={2} -AppName=\"{3}\" -BuildVersion=\"{4}\" -AppLaunch=\"{5}\" -DataAgeThreshold=12", BuildRoot, CloudDir, AppID, AppName, VersionString, AppLaunch);
 				CmdLine += " -AppArgs=\"\"";
@@ -781,6 +780,9 @@ public partial class Project : CommandUtils
 				CmdLine += " -stdout";
 
 				RunAndLog(CmdEnv, BPTExe, CmdLine, Options: ERunOptions.Default | ERunOptions.UTF8Output);
+
+				InternalUtils.SafeCopyFile(SourceManifestPath, BackupManifestPath);
+				InternalUtils.SafeCopyFile(SourceManifestPath, DestManifestPath);
 			}
 			else
 			{
@@ -841,6 +843,23 @@ public partial class Project : CommandUtils
 			}
 		}
 
+		if (Params.CreateChunkInstall)
+		{
+			string ManifestDir = CombinePaths(Params.ChunkInstallDirectory, SC.FinalCookPlatform, "ManifestDir");
+			if (InternalUtils.SafeDirectoryExists(ManifestDir))
+			{
+				foreach (string ManifestFile in Directory.GetFiles(ManifestDir, "*.manifest"))
+				{
+					InternalUtils.SafeDeleteFile(ManifestFile, true);
+				}
+			}
+			string DestDir = CombinePaths(Params.ChunkInstallDirectory, SC.FinalCookPlatform, Params.ChunkInstallVersionString);
+			if (InternalUtils.SafeDirectoryExists(DestDir))
+			{
+				InternalUtils.SafeDeleteDirectory(DestDir); 
+			}
+		}
+
 		for (int ChunkIndex = 0; ChunkIndex < PakResponseFiles.Length; ++ChunkIndex)
 		{
 			var ChunkName = Path.GetFileNameWithoutExtension(ChunkList[ChunkIndex]);
diff --git a/Engine/Source/Runtime/Core/Private/Windows/WindowsPlatformMisc.cpp b/Engine/Source/Runtime/Core/Private/Windows/WindowsPlatformMisc.cpp
--- a/Engine/Source/Runtime/Core/Private/Windows/WindowsPlatformMisc.cpp
+++ b/Engine/Source/Runtime/Core/Private/Windows/WindowsPlatformMisc.cpp
@@ -6,6 +6,8 @@
 #include "WindowsApplication.h"
 #include "EngineVersion.h"
 
+#include "GenericPlatformChunkInstall.h"
+
 // Resource includes.
 #include "Runtime/Launch/Resources/Windows/Resource.h"
 #include "Runtime/Launch/Resources/Version.h"
@@ -2408,3 +2410,26 @@ EConvertibleLaptopMode FWindowsPlatformMisc::GetConvertibleLaptopMode()
 	
 	return EConvertibleLaptopMode::Laptop;
 }
+
+IPlatformChunkInstall* FWindowsPlatformMisc::GetPlatformChunkInstall()
+{
+	static IPlatformChunkInstall* ChunkInstall = nullptr;
+	IPlatformChunkInstallModule* PlatformChunkInstallModule = FModuleManager::LoadModulePtr<IPlatformChunkInstallModule>("HTTPChunkInstaller");
+	if (!ChunkInstall)
+	{
+#if !(WITH_EDITORONLY_DATA || IS_PROGRAM)
+		if (PlatformChunkInstallModule != NULL)
+		{
+			// Attempt to grab the platform installer
+			ChunkInstall = PlatformChunkInstallModule->GetPlatformChunkInstall();
+		}
+		else
+#endif
+		{
+			// Placeholder instance
+			ChunkInstall = new FGenericPlatformChunkInstall();
+		}
+	}
+
+	return ChunkInstall;
+}
diff --git a/Engine/Source/Runtime/Core/Public/GenericPlatform/GenericPlatformChunkInstall.h b/Engine/Source/Runtime/Core/Public/GenericPlatform/GenericPlatformChunkInstall.h
--- a/Engine/Source/Runtime/Core/Public/GenericPlatform/GenericPlatformChunkInstall.h
+++ b/Engine/Source/Runtime/Core/Public/GenericPlatform/GenericPlatformChunkInstall.h
@@ -133,14 +133,14 @@ public:
 	 * @param Delegate		The delegate to call when the chunk is installed.
 	 * @return				False if the delegate was not registered. True on success.
 	 */
-	virtual bool SetChunkInstallDelgate( uint32 ChunkID, FPlatformChunkInstallCompleteDelegate Delegate ) = 0;
+	virtual FDelegateHandle SetChunkInstallDelgate( uint32 ChunkID, FPlatformChunkInstallCompleteDelegate Delegate ) = 0;
 
 	/**
 	* Remove a delegate callback on chunk install completion.
 	* @param ChunkID		The id of the chunk of interest.
 	* @param Delegate		The delegate to remove.
 	*/
-	virtual void RemoveChunkInstallDelgate(uint32 ChunkID, FPlatformChunkInstallCompleteDelegate Delegate) = 0;
+	virtual void RemoveChunkInstallDelgate(uint32 ChunkID, FDelegateHandle Delegate) = 0;
 };
 
 
@@ -235,9 +235,9 @@ public:
 	* @param Delegate		The delegate when the chunk is installed.
 	* @return				False if the delegate was not registered or the chunk is already installed. True on success.
 	*/
-	virtual bool SetChunkInstallDelgate(uint32 ChunkID, FPlatformChunkInstallCompleteDelegate Delegate) override
+	virtual FDelegateHandle SetChunkInstallDelgate(uint32 ChunkID, FPlatformChunkInstallCompleteDelegate Delegate) override
 	{
-		return false;
+		return FDelegateHandle();
 	}
 
 	/**
@@ -246,7 +246,7 @@ public:
 	* @param Delegate		The delegate to remove.
 	* @return				False if the delegate was not registered with a call to SetChunkInstallDelgate. True on success.
 	*/
-	virtual void RemoveChunkInstallDelgate(uint32 ChunkID, FPlatformChunkInstallCompleteDelegate Delegate) override
+	virtual void RemoveChunkInstallDelgate(uint32 ChunkID, FDelegateHandle Delegate) override
 	{
 		return;
 	}
diff --git a/Engine/Source/Runtime/Core/Public/Windows/WindowsPlatformMisc.h b/Engine/Source/Runtime/Core/Public/Windows/WindowsPlatformMisc.h
--- a/Engine/Source/Runtime/Core/Public/Windows/WindowsPlatformMisc.h
+++ b/Engine/Source/Runtime/Core/Public/Windows/WindowsPlatformMisc.h
@@ -259,6 +259,8 @@ struct CORE_API FWindowsPlatformMisc
 	static FString GetOperatingSystemId();
 
 	static EConvertibleLaptopMode GetConvertibleLaptopMode();
+
+	static IPlatformChunkInstall* GetPlatformChunkInstall();
 };
 
 
diff --git a/Engine/Source/Runtime/HTTPChunkInstaller/Private/ChunkInstall.h b/Engine/Source/Runtime/HTTPChunkInstaller/Private/ChunkInstall.h
--- a/Engine/Source/Runtime/HTTPChunkInstaller/Private/ChunkInstall.h
+++ b/Engine/Source/Runtime/HTTPChunkInstaller/Private/ChunkInstall.h
@@ -2,27 +2,37 @@
 
 #pragma once
 
-
-class FChunkInstallTask : public FNonAbandonableTask
+class FChunkInstallTask : public FRunnable
 {
 public:
 	/** Input parameters */
 	FString							ManifestPath;
+	FString							HoldingManifestPath;
 	FString							SrcDir;
 	FString							DestDir;
 	IBuildPatchServicesModule*		BPSModule;
 	IBuildManifestPtr				BuildManifest;
 	bool							bCopy;
 	const TArray<FString>*			CurrentMountPaks;
+	FEvent*							CompleteEvent;				
 	/** Output */
 	TArray<FString>					MountedPaks;
 
 	FChunkInstallTask()
-	{}
+	{
+		CompleteEvent = FPlatformProcess::GetSynchEventFromPool(true);
+	}
+
+	~FChunkInstallTask()
+	{
+		FPlatformProcess::ReturnSynchEventToPool(CompleteEvent);
+		CompleteEvent = nullptr;
+	}
 
-	void Init(FString InManifestPath, FString InSrcDir, FString InDestDir, IBuildPatchServicesModule* InBPSModule, IBuildManifestRef InBuildManifest, const TArray<FString>& InCurrentMountedPaks, bool bInCopy)
+	void SetupWork(FString InManifestPath, FString InHoldingManifestPath, FString InSrcDir, FString InDestDir, IBuildPatchServicesModule* InBPSModule, IBuildManifestRef InBuildManifest, const TArray<FString>& InCurrentMountedPaks, bool bInCopy)
 	{
 		ManifestPath = InManifestPath;
+		HoldingManifestPath = InHoldingManifestPath;
 		SrcDir = InSrcDir;
 		DestDir = InDestDir;
 		BPSModule = InBPSModule;
@@ -31,6 +41,7 @@ public:
 		bCopy = bInCopy;
 
 		MountedPaks.Reset();
+		CompleteEvent->Reset();
 	}
 
 	void DoWork()
@@ -59,11 +70,15 @@ public:
 
 		check(CurrentMountPaks);
 
-		BPSModule->SaveManifestToFile(ManifestPath, BuildManifest.ToSharedRef());
+		TArray<FString> PakFiles;
+		IPlatformFile& PlatformFile = FPlatformFileManager::Get().GetPlatformFile();
+		BPSModule->SaveManifestToFile(ManifestPath, BuildManifest.ToSharedRef(), false);
+		if (PlatformFile.FileExists(*HoldingManifestPath))
+		{
+			PlatformFile.DeleteFile(*HoldingManifestPath);
+		}
 		if (bCopy)
 		{
-			TArray<FString> PakFiles;
-			IPlatformFile& PlatformFile = FPlatformFileManager::Get().GetPlatformFile();
 			if (PlatformFile.DirectoryExists(*DestDir))
 			{
 				PlatformFile.DeleteDirectoryRecursively(*DestDir);
@@ -73,29 +88,47 @@ public:
 			{
 				PlatformFile.DeleteDirectoryRecursively(*SrcDir);
 			}
-			// Find all pak files.
-			FPakSearchVisitor Visitor(PakFiles);
-			PlatformFile.IterateDirectoryRecursively(*DestDir, Visitor);
-			auto PakReadOrderField = BuildManifest->GetCustomField("PakReadOrdering");
-			uint32 PakReadOrder = PakReadOrderField.IsValid() ? (uint32)PakReadOrderField->AsInteger() : 0;
-			for (uint32 PakIndex = 0, PakCount = PakFiles.Num(); PakIndex < PakCount; ++PakIndex)
+		}
+		// Find all pak files.
+		FPakSearchVisitor Visitor(PakFiles);
+		PlatformFile.IterateDirectoryRecursively(*DestDir, Visitor);
+		auto PakReadOrderField = BuildManifest->GetCustomField("PakReadOrdering");
+		uint32 PakReadOrder = PakReadOrderField.IsValid() ? (uint32)PakReadOrderField->AsInteger() : 0;
+		for (uint32 PakIndex = 0, PakCount = PakFiles.Num(); PakIndex < PakCount; ++PakIndex)
+		{
+			if (!CurrentMountPaks->Contains(PakFiles[PakIndex]) && !MountedPaks.Contains(PakFiles[PakIndex]))
 			{
-				if (!CurrentMountPaks->Contains(PakFiles[PakIndex]) && !MountedPaks.Contains(PakFiles[PakIndex]))
+				if (FCoreDelegates::OnMountPak.IsBound())
 				{
-					if (FCoreDelegates::OnMountPak.IsBound())
-					{
-						FCoreDelegates::OnMountPak.Execute(PakFiles[PakIndex], PakReadOrder);
-						MountedPaks.Add(PakFiles[PakIndex]);
-					}
+					FCoreDelegates::OnMountPak.Execute(PakFiles[PakIndex], PakReadOrder);
+					MountedPaks.Add(PakFiles[PakIndex]);
 				}
 			}
-			//Register the install
-			BPSModule->RegisterAppInstallation(BuildManifest.ToSharedRef(), DestDir);
 		}
+		//Register the install
+		BPSModule->RegisterAppInstallation(BuildManifest.ToSharedRef(), DestDir);
+
+		CompleteEvent->Trigger();
+	}
+
+	uint32 Run()
+	{
+		DoWork();
+		return 0;
+	}
+
+	bool IsDone()
+	{
+		return CompleteEvent->Wait(FTimespan(0));
 	}
 
 	static const TCHAR *Name()
 	{
 		return TEXT("FChunkDescovery");
 	}
+
+	FORCEINLINE TStatId GetStatId() const
+	{
+		RETURN_QUICK_DECLARE_CYCLE_STAT(FChunkInstallTask, STATGROUP_ThreadPoolAsyncTasks);
+	}
 };
diff --git a/Engine/Source/Runtime/HTTPChunkInstaller/Private/ChunkSetup.h b/Engine/Source/Runtime/HTTPChunkInstaller/Private/ChunkSetup.h
--- a/Engine/Source/Runtime/HTTPChunkInstaller/Private/ChunkSetup.h
+++ b/Engine/Source/Runtime/HTTPChunkInstaller/Private/ChunkSetup.h
@@ -2,7 +2,7 @@
 
 #pragma once
 
-// Helper class to find all pak files.
+// Helper class to find all pak/mainfests files.
 class FFileSearchVisitor : public IPlatformFile::FDirectoryVisitor
 {
 	FString			 FileWildcard;
@@ -28,23 +28,26 @@ public:
 
 class IBuildPatchServicesModule;
 
-class FChunkSetupTask : public FNonAbandonableTask, public IPlatformFile::FDirectoryVisitor
+class FChunkSetupTask : public FRunnable, public IPlatformFile::FDirectoryVisitor
 {
 public:
 	/** Input parameters */
 	IBuildPatchServicesModule*	BPSModule;
 	FString						InstallDir; // Intermediate directory where installed chunks may be waiting
 	FString						ContentDir; // Directory where installed chunks need to live to be mounted
+	FString						HoldingDir; // Directory where manifest for chunks that are out of date and can be use for updates but not mounted.
 	const TArray<FString>*		CurrentMountPaks; // Array of already mounted Paks
 	/** Output */
+	FEvent*									CompleteEvent;
 	TArray<FString>							MountedPaks;
 	TMultiMap<uint32, IBuildManifestPtr>	InstalledChunks;
+	TMultiMap<uint32, IBuildManifestPtr>	HoldingChunks;
 	/** Working */
 	TArray<FString>		FoundPaks;
 	TArray<FString>		FoundManifests;
 	FFileSearchVisitor	PakVisitor;
 	FFileSearchVisitor	ManifestVisitor;
-	TArray<FString>		DirectoriesToRemove;
+	TArray<FString>		ManifestsToRemove;
 	uint32				Pass;
 
 	FChunkSetupTask()
@@ -52,13 +55,21 @@ public:
 		, CurrentMountPaks(nullptr)
 		, PakVisitor(TEXT("*.pak"), FoundPaks)
 		, ManifestVisitor(TEXT("*.manifest"), FoundManifests)
-	{}
+	{
+		CompleteEvent = FPlatformProcess::GetSynchEventFromPool(true);
+	}
 
-	void Init(IBuildPatchServicesModule* InBPSModule, FString InInstallDir, FString InContentDir, const TArray<FString>& InCurrentMountedPaks)
+	FChunkSetupTask::~FChunkSetupTask()
+	{
+		FPlatformProcess::ReturnSynchEventToPool(CompleteEvent);
+	}
+
+	void SetupWork(IBuildPatchServicesModule* InBPSModule, FString InInstallDir, FString InContentDir, FString InHoldingDir, const TArray<FString>& InCurrentMountedPaks)
 	{
 		BPSModule = InBPSModule;
 		InstallDir = InInstallDir;
 		ContentDir = InContentDir;
+		HoldingDir = InHoldingDir;
 		CurrentMountPaks = &InCurrentMountedPaks;
 
 		Pass = 0;
@@ -66,19 +77,24 @@ public:
 		MountedPaks.Reset();
 		FoundManifests.Reset();
 		FoundPaks.Reset();
-		DirectoriesToRemove.Reset();
+		ManifestsToRemove.Reset();
+
+		CompleteEvent->Reset();
 	}
 
 	void DoWork()
 	{
 		IPlatformFile& PlatformFile = FPlatformFileManager::Get().GetPlatformFile();
 		PlatformFile.IterateDirectory(*InstallDir, *this);
-		for (const auto& ToRemove : DirectoriesToRemove)
+		for (const auto& ToRemove : ManifestsToRemove)
 		{
-			PlatformFile.DeleteDirectoryRecursively(*ToRemove);
+			PlatformFile.DeleteFile(*ToRemove);
 		}
 		++Pass;
 		PlatformFile.IterateDirectory(*ContentDir, *this);
+		++Pass;
+		PlatformFile.IterateDirectory(*HoldingDir, *this);
+		CompleteEvent->Trigger();
 	}
 
 	bool Visit(const TCHAR* FilenameOrDirectory, bool bIsDirectory) override
@@ -95,10 +111,11 @@ public:
 		{
 			return true;
 		}
+		
 		// Only allow one manifest per folder, any more suggests corruption so mark the folder for delete
 		if (FoundManifests.Num() > 1)
 		{
-			DirectoriesToRemove.Add(FilenameOrDirectory);
+			ManifestsToRemove.Append(FoundManifests);
 			return true;
 		}
 		// Load the manifest, so that can be classed as installed
@@ -106,14 +123,14 @@ public:
 		if (!Manifest.IsValid())
 		{
 			//Something is wrong, suggests corruption so mark the folder for delete
-			DirectoriesToRemove.Add(FilenameOrDirectory);
+			ManifestsToRemove.Append(FoundManifests);
 			return true;
 		}
 		auto ChunkIDField = Manifest->GetCustomField("ChunkID");
 		if (!ChunkIDField.IsValid())
 		{
 			//Something is wrong, suggests corruption so mark the folder for delete
-			DirectoriesToRemove.Add(FilenameOrDirectory);
+			ManifestsToRemove.Append(FoundManifests);
 			return true;
 		}
 		auto ChunkPatchField = Manifest->GetCustomField("bIsPatch");
@@ -124,7 +141,7 @@ public:
 		{
 			InstalledChunks.AddUnique(ChunkID, Manifest);
 		}
-		else
+		else if (Pass == 0 && ContentDir != InstallDir)
 		{
 			FString ChunkFdrName = FString::Printf(TEXT("%s%d"), !bIsPatch ? TEXT("base") : TEXT("patch"), ChunkID);
 			FString DestDir = *FPaths::Combine(*ContentDir, *ChunkFdrName);
@@ -135,9 +152,13 @@ public:
 			PlatformFile.CreateDirectoryTree(*DestDir);
 			if (PlatformFile.CopyDirectoryTree(*DestDir, FilenameOrDirectory, true))
 			{
-				DirectoriesToRemove.Add(FilenameOrDirectory);
+				ManifestsToRemove.Add(FilenameOrDirectory);
 			}
 		}
+		else if (Pass == 2)
+		{
+			HoldingChunks.AddUnique(ChunkID, Manifest);
+		}
 
 		return true;
 	}
@@ -146,9 +167,25 @@ public:
 	{
 		return TEXT("FChunkSetup");
 	}
+
+	FORCEINLINE TStatId GetStatId() const
+	{
+		RETURN_QUICK_DECLARE_CYCLE_STAT(FChunkSetupTask, STATGROUP_ThreadPoolAsyncTasks);
+	}
+
+	uint32 Run()
+	{
+		DoWork();
+		return 0;
+	}
+
+	bool IsDone()
+	{
+		return CompleteEvent->Wait(FTimespan(0));
+	}
 };
 
-class FChunkMountTask : public FNonAbandonableTask, public IPlatformFile::FDirectoryVisitor
+class FChunkMountTask : public FRunnable, public IPlatformFile::FDirectoryVisitor
 {
 public:
 	/** Input parameters */
@@ -156,6 +193,7 @@ public:
 	FString						ContentDir; // Directory where installed chunks need to live to be mounted
 	const TArray<FString>*		CurrentMountPaks; // Array of already mounted Paks
 	/** Output */
+	FEvent*									CompleteEvent;
 	TArray<FString>							MountedPaks;
 	/** Working */
 	TArray<FString>		FoundPaks;
@@ -168,9 +206,16 @@ public:
 		, CurrentMountPaks(nullptr)
 		, PakVisitor(TEXT("*.pak"), FoundPaks)
 		, ManifestVisitor(TEXT("*.manifest"), FoundManifests)
-	{}
+	{
+		CompleteEvent = FPlatformProcess::GetSynchEventFromPool(true);
+	}
 
-	void Init(IBuildPatchServicesModule* InBPSModule, FString InContentDir, const TArray<FString>& InCurrentMountedPaks)
+	~FChunkMountTask()
+	{
+		FPlatformProcess::ReturnSynchEventToPool(CompleteEvent);
+	}
+
+	void SetupWork(IBuildPatchServicesModule* InBPSModule, FString InContentDir, const TArray<FString>& InCurrentMountedPaks)
 	{
 		BPSModule = InBPSModule;
 		ContentDir = InContentDir;
@@ -179,12 +224,15 @@ public:
 		MountedPaks.Reset();
 		FoundManifests.Reset();
 		FoundPaks.Reset();
+
+		CompleteEvent->Reset();
 	}
 
 	void DoWork()
 	{
 		IPlatformFile& PlatformFile = FPlatformFileManager::Get().GetPlatformFile();
 		PlatformFile.IterateDirectory(*ContentDir, *this);
+		CompleteEvent->Trigger();
 	}
 
 	bool Visit(const TCHAR* FilenameOrDirectory, bool bIsDirectory) override
@@ -250,4 +298,20 @@ public:
 	{
 		return TEXT("FChunkSetup");
 	}
+
+	FORCEINLINE TStatId GetStatId() const
+	{
+		RETURN_QUICK_DECLARE_CYCLE_STAT(FChunkMountTask, STATGROUP_ThreadPoolAsyncTasks);
+	}
+
+	uint32 Run()
+	{
+		DoWork();
+		return 0;
+	}
+
+	bool IsDone()
+	{
+		return CompleteEvent->Wait(FTimespan(0));
+	}
 };
\ No newline at end of file
diff --git a/Engine/Source/Runtime/HTTPChunkInstaller/Private/HTTPChunkInstaller.cpp b/Engine/Source/Runtime/HTTPChunkInstaller/Private/HTTPChunkInstaller.cpp
--- a/Engine/Source/Runtime/HTTPChunkInstaller/Private/HTTPChunkInstaller.cpp
+++ b/Engine/Source/Runtime/HTTPChunkInstaller/Private/HTTPChunkInstaller.cpp
@@ -52,66 +52,10 @@ FHTTPChunkInstall::FHTTPChunkInstall()
 	, InstallerState(ChunkInstallState::Setup)
 	, InstallSpeed(EChunkInstallSpeed::Fast)
 	, bDebugNoInstalledRequired(false)
+	, bFirstRun(true)
+	, bSystemInitialised(false)
 {
-	BPSModule = GetBuildPatchServices();
-
-#if !UE_BUILD_SHIPPING
-	const TCHAR* CmdLine = FCommandLine::Get();
-	bool Result = FParse::Param(CmdLine,TEXT("Pak")) || FParse::Param(CmdLine,TEXT("Signedpak")) || FParse::Param(CmdLine,TEXT("Signed")) || FParse::Param(CmdLine, TEXT("TestChunkInstall"));
-	if(!FPlatformProperties::RequiresCookedData() || !Result || FParse::Param(CmdLine,TEXT("NoPak")) || FParse::Param(CmdLine, TEXT("NoChunkInstall")))
-	{		
-		bDebugNoInstalledRequired = true;
-	}
-#endif
-
-	// Grab the title file interface
-	FString TitleFileSource;
-	bool bValidTitleFileSource = GConfig->GetString(TEXT("HTTPChunkInstall"), TEXT("TitleFileSource"), TitleFileSource, GEngineIni);
-	if (bValidTitleFileSource && TitleFileSource != TEXT("Local"))
-	{
-		OnlineTitleFile = Online::GetTitleFileInterface(*TitleFileSource);
-	}
-	else
-	{
-		FString LocalTileFileDirectory = FPaths::GameConfigDir();
-		GConfig->GetString(TEXT("HTTPChunkInstall"), TEXT("LocalTitleFileDirectory"), LocalTileFileDirectory, GEngineIni);
-		OnlineTitleFile = MakeShareable(new FLocalTitleFile(LocalTileFileDirectory));
-	}
-	CloudDir	= FPaths::Combine(*FPaths::GameContentDir()	, TEXT("Cloud"));
-	StageDir	= FPaths::Combine(*FPaths::GameSavedDir()	, TEXT("Chunks"), TEXT("Staged"));
-	InstallDir	= FPaths::Combine(*FPaths::GameSavedDir()	, TEXT("Chunks"), TEXT("Installed"));
-	BackupDir	= FPaths::Combine(*FPaths::GameSavedDir()	, TEXT("Chunks"), TEXT("Backup"));
-	CacheDir	= FPaths::Combine(*FPaths::GameSavedDir()	, TEXT("Chunks"), TEXT("Cache"));
-	ContentDir  = FPaths::Combine(*FPaths::GameContentDir() , TEXT("Chunks"));
-
-	FString TmpString1;
-	FString TmpString2;
-	if (GConfig->GetString(TEXT("HTTPChunkInstall"), TEXT("CloudDirectory"), TmpString1, GEngineIni))
-	{
-		CloudDir = TmpString1;
-	}
-	if ((GConfig->GetString(TEXT("HTTPChunkInstall"), TEXT("CloudProtocol"), TmpString1, GEngineIni)) && (GConfig->GetString(TEXT("HTTPChunkInstall"), TEXT("CloudDomain"), TmpString2, GEngineIni)))
-	{
-		CloudDir = FString::Printf(TEXT("%s://%s"), *TmpString1, *TmpString2);
-	}
-	if (GConfig->GetString(TEXT("HTTPChunkInstall"), TEXT("StageDirectory"), TmpString1, GEngineIni))
-	{
-		StageDir = TmpString1;
-	}
-	if (GConfig->GetString(TEXT("HTTPChunkInstall"), TEXT("InstallDirectory"), TmpString1, GEngineIni))
-	{
-		InstallDir = TmpString1;
-	}
-	if (GConfig->GetString(TEXT("HTTPChunkInstall"), TEXT("BackupDirectory"), TmpString1, GEngineIni))
-	{
-		BackupDir = TmpString1;
-	}
-	if (GConfig->GetString(TEXT("HTTPChunkInstall"), TEXT("ContentDirectory"), TmpString1, GEngineIni))
-	{
-		ContentDir = TmpString1;
-	}
 
-	bFirstRun = true;
 }
 
 
@@ -126,27 +70,39 @@ FHTTPChunkInstall::~FHTTPChunkInstall()
 
 bool FHTTPChunkInstall::Tick(float DeltaSeconds)
 {
+	if (!bSystemInitialised)
+	{
+		InitialiseSystem();
+	}
+
 	switch (InstallerState)
 	{
 	case ChunkInstallState::Setup:
 		{
 			check(OnlineTitleFile.IsValid());
-			OnlineTitleFile->AddOnEnumerateFilesCompleteDelegate(FOnEnumerateFilesCompleteDelegate::CreateRaw(this,&FHTTPChunkInstall::OSSEnumerateFilesComplete));
-			OnlineTitleFile->AddOnReadFileCompleteDelegate(FOnReadFileCompleteDelegate::CreateRaw(this,&FHTTPChunkInstall::OSSReadFileComplete));
-			ChunkSetupTask.GetTask().Init(BPSModule, InstallDir, ContentDir, MountedPaks);
-			ChunkSetupTask.StartBackgroundTask();
+			EnumFilesCompleteHandle = OnlineTitleFile->AddOnEnumerateFilesCompleteDelegate_Handle(FOnEnumerateFilesCompleteDelegate::CreateRaw(this,&FHTTPChunkInstall::OSSEnumerateFilesComplete));
+			ReadFileCompleteHandle = OnlineTitleFile->AddOnReadFileCompleteDelegate_Handle(FOnReadFileCompleteDelegate::CreateRaw(this,&FHTTPChunkInstall::OSSReadFileComplete));
+			ChunkSetupTask.SetupWork(BPSModule, InstallDir, ContentDir, HoldingDir, MountedPaks);
+			ChunkSetupTaskThread.Reset(FRunnableThread::Create(&ChunkSetupTask, TEXT("Chunk descovery thread")));
 			InstallerState = ChunkInstallState::SetupWait;
 		} break;
 	case ChunkInstallState::SetupWait:
 		{
 			if (ChunkSetupTask.IsDone())
 			{
-				for (auto It = ChunkSetupTask.GetTask().InstalledChunks.CreateConstIterator(); It; ++It)
+				ChunkSetupTaskThread->WaitForCompletion();
+				ChunkSetupTaskThread.Reset();
+				for (auto It = ChunkSetupTask.InstalledChunks.CreateConstIterator(); It; ++It)
 				{
 					UE_LOG(LogHTTPChunkInstaller, Log, TEXT("Adding Chunk %d to installed manifests"), It.Key());
 					InstalledManifests.Add(It.Key(), It.Value());
 				}
-				MountedPaks.Append(ChunkSetupTask.GetTask().MountedPaks);
+				for (auto It = ChunkSetupTask.HoldingChunks.CreateConstIterator(); It; ++It)
+				{
+					UE_LOG(LogHTTPChunkInstaller, Log, TEXT("Adding Chunk %d to holding manifests"), It.Key());
+					PrevInstallManifests.Add(It.Key(), It.Value());
+				}
+				MountedPaks.Append(ChunkSetupTask.MountedPaks);
 				InstallerState = ChunkInstallState::QueryRemoteManifests;
 			}
 		} break;
@@ -189,8 +145,8 @@ bool FHTTPChunkInstall::Tick(float DeltaSeconds)
 				if (!IsDataInFileCache(TitleFilesToRead[0].Hash))
 				{
 					UE_LOG(LogHTTPChunkInstaller, Log, TEXT("Reading manifest %s from remote source"), *TitleFilesToRead[0].FileName);
-					OnlineTitleFile->ReadFile(TitleFilesToRead[0].DLName);
 					InstallerState = ChunkInstallState::WaitingOnRead;
+					OnlineTitleFile->ReadFile(TitleFilesToRead[0].DLName);
 				}
 				else
 				{
@@ -236,8 +192,8 @@ bool FHTTPChunkInstall::Tick(float DeltaSeconds)
 			{
 				if (bFirstRun)
 				{
-					ChunkMountTask.GetTask().Init(BPSModule, ContentDir, MountedPaks);
-					ChunkMountTask.StartBackgroundTask();
+					ChunkMountTask.SetupWork(BPSModule, ContentDir, MountedPaks);
+					ChunkMountTaskThread.Reset(FRunnableThread::Create(&ChunkMountTask, TEXT("Chunk mounting thread")));
 				}
 				InstallerState = ChunkInstallState::PostSetup;
 			}
@@ -252,7 +208,9 @@ bool FHTTPChunkInstall::Tick(float DeltaSeconds)
 			{
 				if (ChunkMountTask.IsDone())
 				{
-					MountedPaks.Append(ChunkMountTask.GetTask().MountedPaks);
+					ChunkMountTaskThread->WaitForCompletion();
+					ChunkMountTaskThread.Reset();
+					MountedPaks.Append(ChunkMountTask.MountedPaks);
 					UE_LOG(LogHTTPChunkInstaller, Log, TEXT("Completed First Run"));
 					bFirstRun = false;
 				}
@@ -277,12 +235,13 @@ bool FHTTPChunkInstall::Tick(float DeltaSeconds)
 			{
 				InstallService.Reset();
 			}
+			ChunkCopyInstallThread.Reset();
 			check(RemoteManifests.Find(InstallingChunkID));
 			UE_LOG(LogHTTPChunkInstaller, Log, TEXT("Adding Chunk %d to installed manifests"), InstallingChunkID);
 			InstalledManifests.Add(InstallingChunkID, InstallingChunkManifest);
 			UE_LOG(LogHTTPChunkInstaller, Log, TEXT("Removing Chunk %d from remote manifests"), InstallingChunkID);
 			RemoteManifests.Remove(InstallingChunkID, InstallingChunkManifest);
-			MountedPaks.Append(ChunkCopyInstall.GetTask().MountedPaks);
+			MountedPaks.Append(ChunkCopyInstall.MountedPaks);
 			if (!RemoteManifests.Contains(InstallingChunkID))
 			{
 				// No more manifests relating to the chunk ID are left to install.
@@ -293,9 +252,8 @@ bool FHTTPChunkInstall::Tick(float DeltaSeconds)
 					FoundDelegate->Broadcast(InstallingChunkID);
 				}
 			}
-			InstallingChunkID = -1;
-			InstallingChunkManifest.Reset();
-			InstallerState = ChunkInstallState::Idle;
+			EndInstall();
+
 		} break;
 	case ChunkInstallState::Installing:
 	case ChunkInstallState::RequestingTitleFiles:
@@ -331,7 +289,7 @@ void FHTTPChunkInstall::UpdatePendingInstallQueue()
 			auto ChunkIDField = ChunkManifest->GetCustomField("ChunkID");
 			if (ChunkIDField.IsValid())
 			{
-				BeginChunkInstall(NextChunk.ChunkID, ChunkManifest);
+				BeginChunkInstall(NextChunk.ChunkID, ChunkManifest, FindPreviousInstallManifest(ChunkManifest));
 			}
 			else
 			{
@@ -354,7 +312,7 @@ void FHTTPChunkInstall::UpdatePendingInstallQueue()
 				auto ChunkIDField = ChunkManifest->GetCustomField("ChunkID");
 				if (ChunkIDField.IsValid())
 				{
-					BeginChunkInstall(ChunkIDField->AsInteger(), ChunkManifest);
+					BeginChunkInstall(ChunkIDField->AsInteger(), ChunkManifest, FindPreviousInstallManifest(ChunkManifest));
 					return;
 				}
 			}
@@ -372,7 +330,13 @@ EChunkLocation::Type FHTTPChunkInstall::GetChunkLocation(uint32 ChunkID)
 	}
 #endif
 
-	if (bFirstRun)
+	// Safe to assume Chunk0 is ready
+	if (ChunkID == 0)
+	{
+		return EChunkLocation::BestLocation;
+	}
+
+	if (bFirstRun || !bSystemInitialised)
 	{
 		/** Still waiting on setup to finish, report that nothing is installed yet... */
 		return EChunkLocation::NotAvailable;
@@ -403,7 +367,13 @@ float FHTTPChunkInstall::GetChunkProgress(uint32 ChunkID,EChunkProgressReporting
 	}
 #endif
 
-	if (bFirstRun)
+	// Safe to assume Chunk0 is ready
+	if (ChunkID == 0)
+	{
+		return 100.f;
+	}
+
+	if (bFirstRun || !bSystemInitialised)
 	{
 		/** Still waiting on setup to finish, report that nothing is installed yet... */
 		return 0.f;
@@ -452,14 +422,15 @@ void FHTTPChunkInstall::OSSInstallComplete(bool bSuccess, IBuildManifestRef Buil
 		if (!BuildChunkFolderName(BuildManifest, ChunkFdrName, ManifestName, ChunkID, bIsPatch))
 		{
 			//Something bad has happened, bail
-			InstallerState = ChunkInstallState::Idle;
+			EndInstall();
 			return;
 		}
 		UE_LOG(LogHTTPChunkInstaller, Log, TEXT("Chunk %d install complete, preparing to copy to content directory"), ChunkID);
 		FString ManifestPath = FPaths::Combine(*InstallDir, *ChunkFdrName, *ManifestName);
+		FString HoldingManifestPath = FPaths::Combine(*HoldingDir, *ChunkFdrName, *ManifestName);
 		FString SrcDir = FPaths::Combine(*InstallDir, *ChunkFdrName);
 		FString DestDir = FPaths::Combine(*ContentDir, *ChunkFdrName);
-		bool bCopyDir = true; 
+		bool bCopyDir = InstallDir != ContentDir; 
 		TArray<IBuildManifestPtr> FoundManifests;
 		InstalledManifests.MultiFind(ChunkID, FoundManifests);
 		for (const auto& It : FoundManifests)
@@ -471,11 +442,16 @@ void FHTTPChunkInstall::OSSInstallComplete(bool bSuccess, IBuildManifestRef Buil
 				bCopyDir = false;
 			}
 		}
-		ChunkCopyInstall.GetTask().Init(ManifestPath, SrcDir, DestDir, BPSModule, BuildManifest, MountedPaks, bCopyDir);
+		ChunkCopyInstall.SetupWork(ManifestPath, HoldingManifestPath, SrcDir, DestDir, BPSModule, BuildManifest, MountedPaks, bCopyDir);
 		UE_LOG(LogHTTPChunkInstaller, Log, TEXT("Copying Chunk %d to content directory"), ChunkID);
-		ChunkCopyInstall.StartBackgroundTask();
+		ChunkCopyInstallThread.Reset(FRunnableThread::Create(&ChunkCopyInstall, TEXT("Chunk Install Copy Thread")));
+		InstallerState = ChunkInstallState::CopyToContent;
+	}
+	else
+	{
+		//Something bad has happened, return to the Idle state. We'll re-attempt the install
+		EndInstall();
 	}
-	InstallerState = ChunkInstallState::CopyToContent;
 }
 
 void FHTTPChunkInstall::ParseTitleFileManifest(const FString& ManifestFileHash)
@@ -536,9 +512,13 @@ void FHTTPChunkInstall::ParseTitleFileManifest(const FString& ManifestFileHash)
 					if (BuildChunkFolderName(InstalledManifest.ToSharedRef(), ChunkFdrName, ManifestName, ChunkID, bIsPatch))
 					{
 						FString ManifestPath = FPaths::Combine(*ContentDir, *ChunkFdrName, *ManifestName);
+						FString HoldingPath = FPaths::Combine(*HoldingDir, *ChunkFdrName, *ManifestName);
 						IPlatformFile& PlatformFile = FPlatformFileManager::Get().GetPlatformFile();
-						PlatformFile.DeleteFile(*ManifestPath);
+						PlatformFile.CreateDirectoryTree(*FPaths::Combine(*HoldingDir, *ChunkFdrName));
+						PlatformFile.MoveFile(*HoldingPath, *ManifestPath);
 					}
+					UE_LOG(LogHTTPChunkInstaller, Log, TEXT("Adding Chunk %d to previous installed manifests"), ChunkID);
+					PrevInstallManifests.Add(ChunkID, InstalledManifest);
 					UE_LOG(LogHTTPChunkInstaller, Log, TEXT("Removing Chunk %d from installed manifests"), ChunkID);
 					InstalledManifests.Remove(ChunkID, InstalledManifest);
 				}
@@ -594,23 +574,24 @@ bool FHTTPChunkInstall::PrioritizeChunk(uint32 ChunkID, EChunkPriority::Type Pri
 	return true;
 }
 
-bool FHTTPChunkInstall::SetChunkInstallDelgate(uint32 ChunkID, FPlatformChunkInstallCompleteDelegate Delegate)
+FDelegateHandle FHTTPChunkInstall::SetChunkInstallDelgate(uint32 ChunkID, FPlatformChunkInstallCompleteDelegate Delegate)
 {
 	FPlatformChunkInstallCompleteMultiDelegate* FoundDelegate = DelegateMap.Find(ChunkID);
 	if (FoundDelegate)
 	{
-		FoundDelegate->Add(Delegate);
+		return FoundDelegate->Add(Delegate);
 	}
 	else
 	{
 		FPlatformChunkInstallCompleteMultiDelegate MC;
-		MC.Add(Delegate);
+		auto RetVal = MC.Add(Delegate);
 		DelegateMap.Add(ChunkID, MC);
+		return RetVal;
 	}
-	return false;
+	return FDelegateHandle();
 }
 
-void FHTTPChunkInstall::RemoveChunkInstallDelgate(uint32 ChunkID, FPlatformChunkInstallCompleteDelegate Delegate)
+void FHTTPChunkInstall::RemoveChunkInstallDelgate(uint32 ChunkID, FDelegateHandle Delegate)
 {
 	FPlatformChunkInstallCompleteMultiDelegate* FoundDelegate = DelegateMap.Find(ChunkID);
 	if (!FoundDelegate)
@@ -620,10 +601,11 @@ void FHTTPChunkInstall::RemoveChunkInstallDelgate(uint32 ChunkID, FPlatformChunk
 	FoundDelegate->Remove(Delegate);
 }
 
-void FHTTPChunkInstall::BeginChunkInstall(uint32 ChunkID,IBuildManifestPtr ChunkManifest)
+void FHTTPChunkInstall::BeginChunkInstall(uint32 ChunkID,IBuildManifestPtr ChunkManifest, IBuildManifestPtr PrevInstallChunkManifest)
 {
 	check(ChunkManifest->GetCustomField("ChunkID").IsValid());
 	InstallingChunkID = ChunkID;
+	check(ChunkID > 0 && ChunkID < 100);
 	InstallingChunkManifest = ChunkManifest;
 	IPlatformFile& PlatformFile = FPlatformFileManager::Get().GetPlatformFile();
 	auto PatchField = ChunkManifest->GetCustomField("bIsPatch");
@@ -642,7 +624,7 @@ void FHTTPChunkInstall::BeginChunkInstall(uint32 ChunkID,IBuildManifestPtr Chunk
 	BPSModule->SetCloudDirectory(CloudDir);
 	BPSModule->SetStagingDirectory(ChunkStageDir);
 	UE_LOG(LogHTTPChunkInstaller,Log,TEXT("Starting Chunk %d install"),InstallingChunkID);
-	InstallService = BPSModule->StartBuildInstall(nullptr,ChunkManifest,ChunkInstallDir,FBuildPatchBoolManifestDelegate::CreateRaw(this,&FHTTPChunkInstall::OSSInstallComplete));
+	InstallService = BPSModule->StartBuildInstall(PrevInstallChunkManifest,ChunkManifest,ChunkInstallDir,FBuildPatchBoolManifestDelegate::CreateRaw(this,&FHTTPChunkInstall::OSSInstallComplete));
 	if(InstallSpeed == EChunkInstallSpeed::Paused && !InstallService->IsPaused())
 	{
 		InstallService->TogglePauseInstall();
@@ -699,6 +681,102 @@ bool FHTTPChunkInstall::RemoveDataFromFileCache(const FString& ManifestHash)
 	return false;
 }
 
+void FHTTPChunkInstall::InitialiseSystem()
+{
+	BPSModule = GetBuildPatchServices();
+
+#if !UE_BUILD_SHIPPING
+	const TCHAR* CmdLine = FCommandLine::Get();
+	if (!FPlatformProperties::RequiresCookedData() || FParse::Param(CmdLine, TEXT("NoPak")) || FParse::Param(CmdLine, TEXT("NoChunkInstall")))
+	{
+		bDebugNoInstalledRequired = true;
+	}
+#endif
+
+	// Grab the title file interface
+	FString TitleFileSource;
+	bool bValidTitleFileSource = GConfig->GetString(TEXT("HTTPChunkInstall"), TEXT("TitleFileSource"), TitleFileSource, GEngineIni);
+	if (bValidTitleFileSource && TitleFileSource != TEXT("Local"))
+	{
+		OnlineTitleFile = Online::GetTitleFileInterface(*TitleFileSource);
+	}
+	else
+	{
+		FString LocalTileFileDirectory = FPaths::GameConfigDir();
+		auto bGetConfigDir = GConfig->GetString(TEXT("HTTPChunkInstall"), TEXT("LocalTitleFileDirectory"), LocalTileFileDirectory, GEngineIni);
+		OnlineTitleFile = MakeShareable(new FLocalTitleFile(LocalTileFileDirectory));
+#if !UE_BUILD_SHIPPING
+		bDebugNoInstalledRequired = !bGetConfigDir;
+#endif
+	}
+	CloudDir = FPaths::Combine(*FPaths::GameContentDir(), TEXT("Cloud"));
+	StageDir = FPaths::Combine(*FPaths::GameSavedDir(), TEXT("Chunks"), TEXT("Staged"));
+	InstallDir = FPaths::Combine(*FPaths::GameSavedDir(), TEXT("Chunks"), TEXT("Installed")); // By default this should match ContentDir
+	BackupDir = FPaths::Combine(*FPaths::GameSavedDir(), TEXT("Chunks"), TEXT("Backup"));
+	CacheDir = FPaths::Combine(*FPaths::GameSavedDir(), TEXT("Chunks"), TEXT("Cache"));
+	HoldingDir = FPaths::Combine(*FPaths::GameSavedDir(), TEXT("Chunks"), TEXT("Hold"));
+	ContentDir = FPaths::Combine(*FPaths::GameSavedDir(), TEXT("Chunks"), TEXT("Installed")); // By default this should match InstallDir
+
+	FString TmpString1;
+	FString TmpString2;
+	if (GConfig->GetString(TEXT("HTTPChunkInstall"), TEXT("CloudDirectory"), TmpString1, GEngineIni))
+	{
+		CloudDir = TmpString1;
+	}
+	if ((GConfig->GetString(TEXT("HTTPChunkInstall"), TEXT("CloudProtocol"), TmpString1, GEngineIni)) && (GConfig->GetString(TEXT("HTTPChunkInstall"), TEXT("CloudDomain"), TmpString2, GEngineIni)))
+	{
+		CloudDir = FString::Printf(TEXT("%s://%s"), *TmpString1, *TmpString2);
+	}
+	if (GConfig->GetString(TEXT("HTTPChunkInstall"), TEXT("StageDirectory"), TmpString1, GEngineIni))
+	{
+		StageDir = TmpString1;
+	}
+	if (GConfig->GetString(TEXT("HTTPChunkInstall"), TEXT("InstallDirectory"), TmpString1, GEngineIni))
+	{
+		InstallDir = TmpString1;
+	}
+	if (GConfig->GetString(TEXT("HTTPChunkInstall"), TEXT("BackupDirectory"), TmpString1, GEngineIni))
+	{
+		BackupDir = TmpString1;
+	}
+	if (GConfig->GetString(TEXT("HTTPChunkInstall"), TEXT("ContentDirectory"), TmpString1, GEngineIni))
+	{
+		ContentDir = TmpString1;
+	}
+	if (GConfig->GetString(TEXT("HTTPChunkInstall"), TEXT("HoldingDirectory"), TmpString1, GEngineIni))
+	{
+		HoldingDir = TmpString1;
+	}
+
+	bFirstRun = true;
+	bSystemInitialised = true;
+}
+
+IBuildManifestPtr FHTTPChunkInstall::FindPreviousInstallManifest(const IBuildManifestPtr& ChunkManifest)
+{
+	auto ChunkIDField = ChunkManifest->GetCustomField("ChunkID");
+	if (!ChunkIDField.IsValid())
+	{
+		return IBuildManifestPtr();
+	}
+	auto ChunkID = ChunkIDField->AsInteger();
+	TArray<IBuildManifestPtr> FoundManifests;
+	PrevInstallManifests.MultiFind(ChunkID, FoundManifests);
+	return FoundManifests.Num() == 0 ? IBuildManifestPtr() : FoundManifests[0];
+}
+
+void FHTTPChunkInstall::EndInstall()
+{
+	if (InstallService.IsValid())
+	{
+		//InstallService->CancelInstall();
+		InstallService.Reset();
+	}
+	InstallingChunkID = -1;
+	InstallingChunkManifest.Reset();
+	InstallerState = ChunkInstallState::Idle;
+}
+
 /**
  * Module for the HTTP Chunk Installer
  */
diff --git a/Engine/Source/Runtime/HTTPChunkInstaller/Private/HTTPChunkInstaller.h b/Engine/Source/Runtime/HTTPChunkInstaller/Private/HTTPChunkInstaller.h
--- a/Engine/Source/Runtime/HTTPChunkInstaller/Private/HTTPChunkInstaller.h
+++ b/Engine/Source/Runtime/HTTPChunkInstaller/Private/HTTPChunkInstaller.h
@@ -52,8 +52,8 @@ public:
 	}
 	
 	virtual bool PrioritizeChunk( uint32 ChunkID, EChunkPriority::Type Priority ) override;
-	virtual bool SetChunkInstallDelgate(uint32 ChunkID, FPlatformChunkInstallCompleteDelegate Delegate) override;
-	virtual void RemoveChunkInstallDelgate(uint32 ChunkID, FPlatformChunkInstallCompleteDelegate Delegate);
+	virtual FDelegateHandle SetChunkInstallDelgate(uint32 ChunkID, FPlatformChunkInstallCompleteDelegate Delegate) override;
+	virtual void RemoveChunkInstallDelgate(uint32 ChunkID, FDelegateHandle Delegate) override;
 
 	virtual bool DebugStartNextChunk()
 	{
@@ -63,14 +63,18 @@ public:
 
 	bool Tick(float DeltaSeconds) override;
 
+	void EndInstall();
+
 private:
 
+	void InitialiseSystem();
 	bool AddDataToFileCache(const FString& ManifestHash,const TArray<uint8>& Data);
 	bool IsDataInFileCache(const FString& ManifestHash);
 	bool GetDataFromFileCache(const FString& ManifestHash,TArray<uint8>& Data);
 	bool RemoveDataFromFileCache(const FString& ManifestHash);
 	void UpdatePendingInstallQueue();
-	void BeginChunkInstall(uint32 NextChunkID,IBuildManifestPtr ChunkManifest);
+	IBuildManifestPtr FindPreviousInstallManifest(const IBuildManifestPtr& ChunkManifest);
+	void BeginChunkInstall(uint32 NextChunkID,IBuildManifestPtr ChunkManifest, IBuildManifestPtr PrevInstallChunkManifest);
 	void ParseTitleFileManifest(const FString& ManifestFileHash);
 	bool BuildChunkFolderName(IBuildManifestRef Manifest, FString& ChunkFdrName, FString& ManifestName, uint32& ChunkID, bool& bIsPatch);
 	void OSSEnumerateFilesComplete(bool bSuccess);
@@ -119,10 +123,14 @@ private:
 		}
 	};
 
-	FAsyncTask<FChunkInstallTask>								ChunkCopyInstall;
-	FAsyncTask<FChunkSetupTask>									ChunkSetupTask;
-	FAsyncTask<FChunkMountTask>									ChunkMountTask;
+	FChunkInstallTask											ChunkCopyInstall;
+	TUniquePtr<FRunnableThread>									ChunkCopyInstallThread;
+	FChunkSetupTask												ChunkSetupTask;
+	TUniquePtr<FRunnableThread>									ChunkSetupTaskThread;
+	FChunkMountTask												ChunkMountTask;
+	TUniquePtr<FRunnableThread>									ChunkMountTaskThread;
 	TMultiMap<uint32, IBuildManifestPtr>						InstalledManifests;
+	TMultiMap<uint32, IBuildManifestPtr>						PrevInstallManifests;
 	TMultiMap<uint32, IBuildManifestPtr>						RemoteManifests;
 	TMap<uint32, FPlatformChunkInstallCompleteMultiDelegate>	DelegateMap;
 	TSet<FString>												ManifestsInMemory;
@@ -133,17 +141,21 @@ private:
 	IOnlineTitleFilePtr											OnlineTitleFile;
 	IBuildInstallerPtr											InstallService;
 	IBuildManifestPtr											InstallingChunkManifest;
+	FDelegateHandle												EnumFilesCompleteHandle;
+	FDelegateHandle												ReadFileCompleteHandle;
 	FString														CloudDir;
 	FString														StageDir;
 	FString														InstallDir;
 	FString														BackupDir;
 	FString														ContentDir;
 	FString														CacheDir;
+	FString														HoldingDir;
 	IBuildPatchServicesModule*									BPSModule;
 	uint32														InstallingChunkID;
 	ChunkInstallState											InstallerState;
 	EChunkInstallSpeed::Type									InstallSpeed;
 	bool														bFirstRun;
+	bool														bSystemInitialised;
 #if !UE_BUILD_SHIPPING
 	bool														bDebugNoInstalledRequired;
 #endif
diff --git a/Engine/Source/Runtime/HTTPChunkInstaller/Private/LocalTitleFile.cpp b/Engine/Source/Runtime/HTTPChunkInstaller/Private/LocalTitleFile.cpp
--- a/Engine/Source/Runtime/HTTPChunkInstaller/Private/LocalTitleFile.cpp
+++ b/Engine/Source/Runtime/HTTPChunkInstaller/Private/LocalTitleFile.cpp
@@ -6,7 +6,13 @@
 
 FLocalTitleFile::FLocalTitleFile(const FString& InRootDirectory)
 : RootDirectory(InRootDirectory)
-{ }
+{
+	RootDirectory.Replace(TEXT("\\"), TEXT("/"));
+	if (!RootDirectory.EndsWith(TEXT("/")))
+	{
+		RootDirectory.Append(TEXT("/"));
+	}
+}
 
 
 bool FLocalTitleFile::GetFileContents(const FString& DLName, TArray<uint8>& FileContents)
