Adding Http chunk install options to Editor Packaging menu.
Adding Http chunk installer sync mode and offline mode.

diff --git a/Engine/Source/Editor/MainFrame/Private/Frame/MainFrameActions.cpp b/Engine/Source/Editor/MainFrame/Private/Frame/MainFrameActions.cpp
--- a/Engine/Source/Editor/MainFrame/Private/Frame/MainFrameActions.cpp
+++ b/Engine/Source/Editor/MainFrame/Private/Frame/MainFrameActions.cpp
@@ -575,6 +575,11 @@ void FMainFrameActionCallbacks::PackageProject( const FName InPlatformInfoName )
 		OptionalParams += TEXT(" -CrashReporter");
 	}
 
+	if (PackagingSettings->bBuildHttpChunkInstallData)
+	{
+		OptionalParams += FString::Printf(TEXT(" -manifests -createchunkinstall -chunkinstalldirectory=\"%s\" -chunkinstallversion=%s"), *(PackagingSettings->HttpChunkInstallDataDirectory.Path), *(PackagingSettings->HttpChunkInstallDataVersion));
+	}
+
 	FString ExecutableName = FPlatformProcess::ExecutableName(false);
 #if PLATFORM_WINDOWS
 	// turn UE4editor into UE4editor.cmd
diff --git a/Engine/Source/Editor/UnrealEd/Classes/Settings/ProjectPackagingSettings.h b/Engine/Source/Editor/UnrealEd/Classes/Settings/ProjectPackagingSettings.h
--- a/Engine/Source/Editor/UnrealEd/Classes/Settings/ProjectPackagingSettings.h
+++ b/Engine/Source/Editor/UnrealEd/Classes/Settings/ProjectPackagingSettings.h
@@ -93,6 +93,24 @@ public:
 	UPROPERTY(config, EditAnywhere, Category=Packaging)
 	bool bGenerateChunks;
 
+	/** 
+	 * If enabled, will generate data for HTTP Chunk Installer. This data can be hosted on webserver to be installed at runtime. Requires "Generate Chunks" enabled.
+	 */
+	UPROPERTY(config, EditAnywhere, Category=Packaging)
+	bool bBuildHttpChunkInstallData;
+
+	/** 
+	 * When "Build HTTP Chunk Install Data" is enabled this is the directory where the data will be build to.
+	 */	
+	UPROPERTY(config, EditAnywhere, Category = Packaging)
+	FDirectoryPath HttpChunkInstallDataDirectory;
+
+	/** 
+	 * Version name for HTTP Chunk Install Data.
+	 */
+	UPROPERTY(config, EditAnywhere, Category = Packaging)
+	FString HttpChunkInstallDataVersion;
+
 	/** Specifies whether to include prerequisites of packaged games, such as redistributable operating system components, whenever possible. */
 	UPROPERTY(config, EditAnywhere, Category=Packaging)
 	bool IncludePrerequisites;
diff --git a/Engine/Source/Editor/UnrealEd/Private/Settings/SettingsClasses.cpp b/Engine/Source/Editor/UnrealEd/Private/Settings/SettingsClasses.cpp
--- a/Engine/Source/Editor/UnrealEd/Private/Settings/SettingsClasses.cpp
+++ b/Engine/Source/Editor/UnrealEd/Private/Settings/SettingsClasses.cpp
@@ -501,6 +501,39 @@ void UProjectPackagingSettings::PostEditChangeProperty( FPropertyChangedEvent& P
 			BuildConfiguration = EProjectPackagingBuildConfigurations::PPBC_Shipping;
 		}
 	}
+	else if (Name == FName(TEXT("bGenerateChunks")))
+	{
+		if (bGenerateChunks)
+		{
+			UsePakFile = true;
+		}
+	}
+	else if (Name == FName(TEXT("UsePakFile")))
+	{
+		if (!UsePakFile)
+		{
+			bGenerateChunks = false;
+			bBuildHttpChunkInstallData = false;
+		}
+	}
+	else if (Name == FName(TEXT("bBuildHTTPChunkInstallData")))
+	{
+		if (bBuildHttpChunkInstallData)
+		{
+			UsePakFile = true;
+			bGenerateChunks = true;
+			//Ensure data is something valid
+			if (HttpChunkInstallDataDirectory.Path.IsEmpty())
+			{
+				auto CloudInstallDir = FPaths::ConvertRelativePathToFull(FPaths::GetPath(FPaths::GetProjectFilePath())) / TEXT("ChunkInstall");
+				HttpChunkInstallDataDirectory.Path = CloudInstallDir;
+			}
+			if (HttpChunkInstallDataVersion.IsEmpty())
+			{
+				HttpChunkInstallDataVersion = TEXT("release1");
+			}
+		}
+	}
 }
 
 bool UProjectPackagingSettings::CanEditChange( const UProperty* InProperty ) const
diff --git a/Engine/Source/Runtime/HTTPChunkInstaller/Private/ChunkInstall.h b/Engine/Source/Runtime/HTTPChunkInstaller/Private/ChunkInstall.h
--- a/Engine/Source/Runtime/HTTPChunkInstaller/Private/ChunkInstall.h
+++ b/Engine/Source/Runtime/HTTPChunkInstaller/Private/ChunkInstall.h
@@ -100,7 +100,15 @@ public:
 			{
 				if (FCoreDelegates::OnMountPak.IsBound())
 				{
-					FCoreDelegates::OnMountPak.Execute(PakFiles[PakIndex], PakReadOrder);
+					auto bSuccess = FCoreDelegates::OnMountPak.Execute(PakFiles[PakIndex], PakReadOrder);
+#if !UE_BUILD_SHIPPING
+					if (!bSuccess)
+					{
+						// This can fail because of the sandbox system - which the pak system doesn't understand.
+						auto SandboxedPath = IFileManager::Get().ConvertToAbsolutePathForExternalAppForRead(*PakFiles[PakIndex]);
+						bSuccess = FCoreDelegates::OnMountPak.Execute(SandboxedPath, PakReadOrder);
+					}
+#endif
 					MountedPaks.Add(PakFiles[PakIndex]);
 				}
 			}
diff --git a/Engine/Source/Runtime/HTTPChunkInstaller/Private/ChunkSetup.h b/Engine/Source/Runtime/HTTPChunkInstaller/Private/ChunkSetup.h
--- a/Engine/Source/Runtime/HTTPChunkInstaller/Private/ChunkSetup.h
+++ b/Engine/Source/Runtime/HTTPChunkInstaller/Private/ChunkSetup.h
@@ -192,12 +192,14 @@ public:
 	IBuildPatchServicesModule*	BPSModule;
 	FString						ContentDir; // Directory where installed chunks need to live to be mounted
 	const TArray<FString>*		CurrentMountPaks; // Array of already mounted Paks
+	const TSet<uint32>*			ExpectedChunks; //Manifests expected to be seen. Chunks not in this list are deleted
 	/** Output */
 	FEvent*									CompleteEvent;
 	TArray<FString>							MountedPaks;
 	/** Working */
 	TArray<FString>		FoundPaks;
 	TArray<FString>		FoundManifests;
+	TArray<FString>		ChunkInstallToDestroy;
 	FFileSearchVisitor	PakVisitor;
 	FFileSearchVisitor	ManifestVisitor;
 
@@ -215,15 +217,17 @@ public:
 		FPlatformProcess::ReturnSynchEventToPool(CompleteEvent);
 	}
 
-	void SetupWork(IBuildPatchServicesModule* InBPSModule, FString InContentDir, const TArray<FString>& InCurrentMountedPaks)
+	void SetupWork(IBuildPatchServicesModule* InBPSModule, FString InContentDir, const TArray<FString>& InCurrentMountedPaks, const TSet<uint32>& InExpectedChunks)
 	{
 		BPSModule = InBPSModule;
 		ContentDir = InContentDir;
 		CurrentMountPaks = &InCurrentMountedPaks;
+		ExpectedChunks = &InExpectedChunks;
 
 		MountedPaks.Reset();
 		FoundManifests.Reset();
 		FoundPaks.Reset();
+		ChunkInstallToDestroy.Reset();
 
 		CompleteEvent->Reset();
 	}
@@ -232,6 +236,10 @@ public:
 	{
 		IPlatformFile& PlatformFile = FPlatformFileManager::Get().GetPlatformFile();
 		PlatformFile.IterateDirectory(*ContentDir, *this);
+		for (const auto& Dir : ChunkInstallToDestroy)
+		{
+			PlatformFile.DeleteDirectoryRecursively(*Dir);
+		}
 		CompleteEvent->Trigger();
 	}
 
@@ -268,6 +276,13 @@ public:
 			return true;
 		}
 
+		if (!ExpectedChunks->Find(ChunkIDField->AsInteger())) 
+		{
+			//Add this to the list of chunks to remove
+			ChunkInstallToDestroy.Add(FilenameOrDirectory);
+			return true;
+		}
+
 		FoundPaks.Reset();
 		PlatformFile.IterateDirectoryRecursively(FilenameOrDirectory, PakVisitor);
 		if (FoundPaks.Num() == 0)
@@ -284,8 +299,15 @@ public:
 			{
 				if (FCoreDelegates::OnMountPak.IsBound())
 				{
-					FCoreDelegates::OnMountPak.Execute(PakPath, PakReadOrder);
-					MountedPaks.Add(PakPath);
+					auto bSuccess = FCoreDelegates::OnMountPak.Execute(PakPath, PakReadOrder);
+#if !UE_BUILD_SHIPPING
+					if (!bSuccess)
+					{
+						// This can fail because of the sandbox system - which the pak system doesn't understand.
+						auto SandboxedPath = IFileManager::Get().ConvertToAbsolutePathForExternalAppForRead(*PakPath);
+						bSuccess = FCoreDelegates::OnMountPak.Execute(SandboxedPath, PakReadOrder);
+					}
+#endif
 					//Register the install
 					BPSModule->RegisterAppInstallation(Manifest.ToSharedRef(), FilenameOrDirectory);
 				}
diff --git a/Engine/Source/Runtime/HTTPChunkInstaller/Private/HTTPChunkInstaller.cpp b/Engine/Source/Runtime/HTTPChunkInstaller/Private/HTTPChunkInstaller.cpp
--- a/Engine/Source/Runtime/HTTPChunkInstaller/Private/HTTPChunkInstaller.cpp
+++ b/Engine/Source/Runtime/HTTPChunkInstaller/Private/HTTPChunkInstaller.cpp
@@ -110,12 +110,6 @@ bool FHTTPChunkInstall::Tick(float DeltaSeconds)
 		{			
 			//Now query the title file service for the chunk manifests. This should return the list of expected chunk manifests
 			check(OnlineTitleFile.IsValid());
-			if (InstallSpeed == EChunkInstallSpeed::Paused)
-			{
-				UE_LOG(LogHTTPChunkInstaller, Log, TEXT("Skipping enumerating manifest files because system is paused"));
-				InstallerState = ChunkInstallState::PostSetup;
-				break;
-			}
 			OnlineTitleFile->ClearFiles();
 			InstallerState = ChunkInstallState::RequestingTitleFiles;
 			UE_LOG(LogHTTPChunkInstaller, Log, TEXT("Enumerating manifest files"));
@@ -127,6 +121,7 @@ bool FHTTPChunkInstall::Tick(float DeltaSeconds)
 			TArray<FCloudFileHeader> FileList;
 			TitleFilesToRead.Reset();
 			RemoteManifests.Reset();
+			ExpectedChunks.Empty();
 			OnlineTitleFile->GetFileList(FileList);
 			for (int32 FileIndex = 0, FileCount = FileList.Num(); FileIndex < FileCount; ++FileIndex)
 			{
@@ -192,7 +187,7 @@ bool FHTTPChunkInstall::Tick(float DeltaSeconds)
 			{
 				if (bFirstRun)
 				{
-					ChunkMountTask.SetupWork(BPSModule, ContentDir, MountedPaks);
+					ChunkMountTask.SetupWork(BPSModule, ContentDir, MountedPaks, ExpectedChunks);
 					ChunkMountTaskThread.Reset(FRunnableThread::Create(&ChunkMountTask, TEXT("Chunk mounting thread")));
 				}
 				InstallerState = ChunkInstallState::PostSetup;
@@ -202,6 +197,16 @@ bool FHTTPChunkInstall::Tick(float DeltaSeconds)
 				InstallerState = ChunkInstallState::ReadTitleFiles;
 			}
 		} break;
+	case ChunkInstallState::EnterOfflineMode:
+		{
+			for (auto It = InstalledManifests.CreateConstIterator(); It; ++It)
+			{
+				ExpectedChunks.Add(It.Key());
+			}
+			ChunkMountTask.SetupWork(BPSModule, ContentDir, MountedPaks, ExpectedChunks);
+			ChunkMountTaskThread.Reset(FRunnableThread::Create(&ChunkMountTask, TEXT("Chunk mounting thread")));
+			InstallerState = ChunkInstallState::PostSetup;
+		} break;
 	case ChunkInstallState::PostSetup:
 		{
 			if (bFirstRun)
@@ -401,12 +406,12 @@ float FHTTPChunkInstall::GetChunkProgress(uint32 ChunkID,EChunkProgressReporting
 
 void FHTTPChunkInstall::OSSEnumerateFilesComplete(bool bSuccess)
 {
-	InstallerState = bSuccess ? ChunkInstallState::SearchTitleFiles : ChunkInstallState::QueryRemoteManifests;
+	InstallerState = bSuccess ? ChunkInstallState::SearchTitleFiles : ChunkInstallState::EnterOfflineMode;
 }
 
 void FHTTPChunkInstall::OSSReadFileComplete(bool bSuccess, const FString& Filename)
 {
-	InstallerState = bSuccess ? ChunkInstallState::ReadComplete : ChunkInstallState::ReadTitleFiles;
+	InstallerState = bSuccess ? ChunkInstallState::ReadComplete : ChunkInstallState::EnterOfflineMode;
 }
 
 void FHTTPChunkInstall::OSSInstallComplete(bool bSuccess, IBuildManifestRef BuildManifest)
@@ -479,6 +484,7 @@ void FHTTPChunkInstall::ParseTitleFileManifest(const FString& ManifestFileHash)
 	}
 	//Compare to installed manifests and add to the remote if it needs to be installed.
 	uint32 ChunkID = (uint32)RemoteChunkIDField->AsInteger();
+	ExpectedChunks.Add(ChunkID);
 	TArray<IBuildManifestPtr> FoundManifests;
 	InstalledManifests.MultiFind(ChunkID, FoundManifests);
 	uint32 FoundCount = FoundManifests.Num();
diff --git a/Engine/Source/Runtime/HTTPChunkInstaller/Private/HTTPChunkInstaller.h b/Engine/Source/Runtime/HTTPChunkInstaller/Private/HTTPChunkInstaller.h
--- a/Engine/Source/Runtime/HTTPChunkInstaller/Private/HTTPChunkInstaller.h
+++ b/Engine/Source/Runtime/HTTPChunkInstaller/Private/HTTPChunkInstaller.h
@@ -88,6 +88,7 @@ private:
 		Setup,
 		SetupWait,
 		QueryRemoteManifests,
+		EnterOfflineMode,
 		MoveInstalledChunks,
 		RequestingTitleFiles,
 		SearchTitleFiles,
@@ -134,6 +135,7 @@ private:
 	TMultiMap<uint32, IBuildManifestPtr>						RemoteManifests;
 	TMap<uint32, FPlatformChunkInstallCompleteMultiDelegate>	DelegateMap;
 	TSet<FString>												ManifestsInMemory;
+	TSet<uint32>												ExpectedChunks;
 	TArray<FCloudFileHeader>									TitleFilesToRead;
 	TArray<uint8>												FileContentBuffer;
 	TArray<FChunkPrio>											PriorityQueue;
diff --git a/Engine/Source/Runtime/HTTPChunkInstaller/Private/HTTPChunkInstallerPrivatePCH.h b/Engine/Source/Runtime/HTTPChunkInstaller/Private/HTTPChunkInstallerPrivatePCH.h
--- a/Engine/Source/Runtime/HTTPChunkInstaller/Private/HTTPChunkInstallerPrivatePCH.h
+++ b/Engine/Source/Runtime/HTTPChunkInstaller/Private/HTTPChunkInstallerPrivatePCH.h
@@ -7,5 +7,6 @@
 #include "ModuleManager.h"
 #include "Engine.h"
 #include "OnlineSubsystemUtils.h"
+#include "FileManagerGeneric.h"
 
 DECLARE_LOG_CATEGORY_EXTERN(LogHTTPChunkInstaller, Log, All);
\ No newline at end of file
